Checkpoint (Combined Team Work):
- Subsonic Event Handler Asset, Factory, Editor, & Definition stubs
- Add Subsonic plugin to applicable test projects dev/testing

#rnx

#Synced-CL 49790744

diff --git a/Engine/Plugins/Experimental/Subsonic/Source/Subsonic/Private/SubsonicEventHandler.cpp b/Engine/Plugins/Experimental/Subsonic/Source/Subsonic/Private/SubsonicEventHandler.cpp
new file mode 100644
--- /dev/null
+++ b/Engine/Plugins/Experimental/Subsonic/Source/Subsonic/Private/SubsonicEventHandler.cpp
@@ -0,0 +1,14 @@
+// Copyright Epic Games, Inc. All Rights Reserved.
+
+#include "SubsonicEventHandler.h"
+
+void FSubsonicEventDefinition::Execute() const
+{
+	for (const FSubsonicEventActionData& ActionData : Actions)
+	{
+		if (const FSubsonicEventActionBase* Action = ActionData.Action.GetPtr())
+		{
+			Action->Execute();
+		}
+	}
+}
\ No newline at end of file
diff --git a/Engine/Plugins/Experimental/Subsonic/Source/Subsonic/Private/SubsonicEventInstanceImpl.cpp b/Engine/Plugins/Experimental/Subsonic/Source/Subsonic/Private/SubsonicEventInstanceImpl.cpp
new file mode 100644
--- /dev/null
+++ b/Engine/Plugins/Experimental/Subsonic/Source/Subsonic/Private/SubsonicEventInstanceImpl.cpp
@@ -0,0 +1,27 @@
+// Copyright Epic Games, Inc. All Rights Reserved.
+
+#include "SubsonicEventInstanceImpl.h"
+
+#include "Components/AudioComponent.h"
+#include "Engine/World.h"
+#include "Kismet/GameplayStatics.h"
+
+FSubsonicEventInstanceImpl::FSubsonicEventInstanceImpl()
+{
+}
+
+void FSubsonicEventInstanceImpl::PlaySound(const UObject* WorldContextObject, USoundBase* Sound)
+{
+	AudioComponents.Add(UGameplayStatics::SpawnSound2D(WorldContextObject, Sound));
+}
+
+void FSubsonicEventInstanceImpl::Stop()
+{
+	for (TWeakObjectPtr<UAudioComponent> AudioComponent : AudioComponents)
+	{
+		if (AudioComponent.IsValid())
+		{
+			AudioComponent->Stop();
+		}
+	}
+}
\ No newline at end of file
diff --git a/Engine/Plugins/Experimental/Subsonic/Source/Subsonic/Private/SubsonicEventInstanceImpl.h b/Engine/Plugins/Experimental/Subsonic/Source/Subsonic/Private/SubsonicEventInstanceImpl.h
new file mode 100644
--- /dev/null
+++ b/Engine/Plugins/Experimental/Subsonic/Source/Subsonic/Private/SubsonicEventInstanceImpl.h
@@ -0,0 +1,21 @@
+// Copyright Epic Games, Inc. All Rights Reserved.
+
+#pragma once
+
+#include "UObject/WeakObjectPtr.h"
+
+class UAudioComponent;
+class USoundBase;
+
+class FSubsonicEventInstanceImpl
+{
+public:
+	FSubsonicEventInstanceImpl();
+
+	void PlaySound(const UObject* WorldContextObject, USoundBase* Sound);
+
+	void Stop();
+
+private:
+	TArray<TWeakObjectPtr<UAudioComponent>> AudioComponents;
+};
diff --git a/Engine/Plugins/Experimental/Subsonic/Source/Subsonic/Private/SubsonicSubsystem.cpp b/Engine/Plugins/Experimental/Subsonic/Source/Subsonic/Private/SubsonicSubsystem.cpp
new file mode 100644
--- /dev/null
+++ b/Engine/Plugins/Experimental/Subsonic/Source/Subsonic/Private/SubsonicSubsystem.cpp
@@ -0,0 +1,95 @@
+// Copyright Epic Games, Inc. All Rights Reserved.
+
+#include "SubsonicSubsystem.h"
+
+#include "Engine/GameInstance.h"
+#include "SubsonicEventInstanceImpl.h"
+#include "SubsonicTempDataAsset.h"
+#include "UObject/GCObject.h"
+
+class FSubsonicSubsystemImpl : FGCObject
+{
+public:
+	FSubsonicEventInstance PlayEvent(const UObject* WorldContextObject, const FName EventName);
+	void StopEvent(FSubsonicEventInstance EventInstance);
+	void RegisterHandlers(USubsonicTempDataAsset* Handlers);
+
+private:
+	virtual void AddReferencedObjects(FReferenceCollector& Collector) override
+	{
+		Collector.AddStableReferenceArray(&RegisteredHandlers);
+	}
+	virtual FString GetReferencerName() const override
+	{
+		return "FSubsonicSubsystemImpl";
+	}
+
+	uint32 NextHandle = 1;
+
+	TMap<FSubsonicEventInstance, FSubsonicEventInstanceImpl> ActiveInstances;
+
+	TArray<TObjectPtr<USubsonicTempDataAsset>> RegisteredHandlers;
+};
+
+FSubsonicEventInstance FSubsonicSubsystemImpl::PlayEvent(const UObject* WorldContextObject, const FName EventName)
+{
+	FSubsonicEventInstance NewHandle;
+	NewHandle.RawHandle = NextHandle++;
+
+	FSubsonicEventInstanceImpl& NewInstance = ActiveInstances.Emplace(NewHandle, {});
+	for (USubsonicTempDataAsset* Handlers : RegisteredHandlers)
+	{
+		TObjectPtr<USoundBase>* MappedSound = Handlers->EventHandlers.Find(EventName);
+		if (MappedSound && *MappedSound)
+		{
+			NewInstance.PlaySound(WorldContextObject, *MappedSound);
+		}
+	}
+
+	return NewHandle;
+}
+
+void FSubsonicSubsystemImpl::StopEvent(FSubsonicEventInstance EventInstance)
+{
+	FSubsonicEventInstanceImpl* ActiveInstance = ActiveInstances.Find(EventInstance);
+	if (ActiveInstance)
+	{
+		ActiveInstance->Stop();
+	}
+}
+
+void FSubsonicSubsystemImpl::RegisterHandlers(USubsonicTempDataAsset* Handlers)
+{
+	if (Handlers)
+	{
+		RegisteredHandlers.Add(Handlers);
+	}
+}
+
+void USubsonicSubsystem::Initialize(FSubsystemCollectionBase& Collection)
+{
+	Super::Initialize(Collection);
+	Impl = MakePimpl<FSubsonicSubsystemImpl>();
+}
+
+void USubsonicSubsystem::Deinitialize()
+{
+	Impl.Reset();
+	Super::Deinitialize();
+}
+
+FSubsonicEventInstance USubsonicSubsystem::PlayEvent(const UObject* WorldContextObject, const FName EventName)
+{
+	return Impl->PlayEvent(WorldContextObject, EventName);
+
+}
+
+void USubsonicSubsystem::StopEvent(FSubsonicEventInstance EventInstance)
+{
+	Impl->StopEvent(EventInstance);
+}
+
+void USubsonicSubsystem::RegisterHandlers(USubsonicTempDataAsset* Handlers)
+{
+	Impl->RegisterHandlers(Handlers);
+}
\ No newline at end of file
diff --git a/Engine/Plugins/Experimental/Subsonic/Source/Subsonic/Private/SubsonicTempDataAsset.h b/Engine/Plugins/Experimental/Subsonic/Source/Subsonic/Private/SubsonicTempDataAsset.h
new file mode 100644
--- /dev/null
+++ b/Engine/Plugins/Experimental/Subsonic/Source/Subsonic/Private/SubsonicTempDataAsset.h
@@ -0,0 +1,22 @@
+// Copyright Epic Games, Inc. All Rights Reserved.
+
+#pragma once
+
+#include "Engine/DataAsset.h"
+
+#include "SubsonicTempDataAsset.generated.h"
+
+class USoundBase;
+
+/**
+* Dumb placeholder for subsonic event handlers
+*/
+UCLASS(MinimalAPI, BlueprintType)
+class USubsonicTempDataAsset final : public UDataAsset
+{
+	GENERATED_BODY()
+
+public:
+	UPROPERTY(EditAnywhere, Category="Subsonic")
+	TMap<FName, TObjectPtr<USoundBase>> EventHandlers;
+};
diff --git a/Engine/Plugins/Experimental/Subsonic/Source/Subsonic/Private/SubsonicTypes.cpp b/Engine/Plugins/Experimental/Subsonic/Source/Subsonic/Private/SubsonicTypes.cpp
new file mode 100644
--- /dev/null
+++ b/Engine/Plugins/Experimental/Subsonic/Source/Subsonic/Private/SubsonicTypes.cpp
@@ -0,0 +1,8 @@
+// Copyright Epic Games, Inc. All Rights Reserved.
+
+#include "SubsonicTypes.h"
+
+uint32 GetTypeHash(FSubsonicEventInstance Instance)
+{
+	return GetTypeHash(Instance.RawHandle);
+}
diff --git a/Engine/Plugins/Experimental/Subsonic/Source/Subsonic/Public/SubsonicEventHandler.h b/Engine/Plugins/Experimental/Subsonic/Source/Subsonic/Public/SubsonicEventHandler.h
new file mode 100644
--- /dev/null
+++ b/Engine/Plugins/Experimental/Subsonic/Source/Subsonic/Public/SubsonicEventHandler.h
@@ -0,0 +1,63 @@
+// Copyright Epic Games, Inc. All Rights Reserved.
+#pragma once
+
+#include "StructUtils/InstancedStruct.h"
+
+#include "SubsonicEventHandler.generated.h"
+
+
+USTRUCT(MinimalAPI, BlueprintType, meta = (Hidden))
+struct FSubsonicEventActionBase
+{
+	GENERATED_BODY()
+
+	virtual ~FSubsonicEventActionBase() = default;
+
+	virtual void Execute() const { }
+};
+
+
+USTRUCT(MinimalAPI, BlueprintType)
+struct FSubsonicEventActionData
+{
+	GENERATED_BODY()
+
+	UPROPERTY(EditAnywhere, Category = "Subsonic")
+	TInstancedStruct<FSubsonicEventActionBase> Action;
+};
+
+
+USTRUCT(MinimalAPI, BlueprintType)
+struct FSubsonicEventDefinition
+{
+	GENERATED_BODY()
+
+	virtual ~FSubsonicEventDefinition() = default;
+
+	void Execute() const;
+
+	UPROPERTY(EditAnywhere, Category = "Subsonic")
+	TArray<FSubsonicEventActionData> Actions;
+};
+
+USTRUCT(MinimalAPI, BlueprintType)
+struct FSubsonicEventData
+{
+	GENERATED_BODY()
+
+	UPROPERTY(EditAnywhere, Category = "Subsonic")
+	TInstancedStruct<FSubsonicEventDefinition> Event;
+};
+
+/**
+* Dumb placeholder for subsonic event handlers
+*/
+UCLASS(MinimalAPI, BlueprintType)
+class USubsonicEventHandler final : public UObject
+{
+	GENERATED_BODY()
+
+public:
+	UPROPERTY(EditAnywhere, Category = "Subsonic")
+	TMap<FName, FSubsonicEventData> EventHandlers;
+};
diff --git a/Engine/Plugins/Experimental/Subsonic/Source/Subsonic/Public/SubsonicSubsystem.h b/Engine/Plugins/Experimental/Subsonic/Source/Subsonic/Public/SubsonicSubsystem.h
new file mode 100644
--- /dev/null
+++ b/Engine/Plugins/Experimental/Subsonic/Source/Subsonic/Public/SubsonicSubsystem.h
@@ -0,0 +1,42 @@
+// Copyright Epic Games, Inc. All Rights Reserved.
+
+#pragma once
+
+#include "SubsonicTypes.h"
+#include "Subsystems/EngineSubsystem.h"
+
+#include "SubsonicSubsystem.generated.h"
+
+#define UE_API SUBSONIC_API
+
+class FSubsonicEventInstanceImpl;
+class FSubsonicSubsystemImpl;
+class USubsonicTempDataAsset;
+
+/**
+ * Interface to the Subsonic plugin.
+ */
+UCLASS(MinimalAPI, BlueprintType)
+class USubsonicSubsystem final : public UEngineSubsystem
+{
+	GENERATED_BODY()
+
+public:
+	//~ Begin USubsystem interface
+	UE_API virtual void Initialize(FSubsystemCollectionBase& Collection) override;
+	UE_API virtual void Deinitialize() override;
+	//~ End USubsystem interface
+
+	UFUNCTION(BlueprintCallable, Category = "Subsonic")
+	UE_API void RegisterHandlers(USubsonicTempDataAsset* Handlers);
+
+	UFUNCTION(BlueprintCallable, Category = "Subsonic", meta=(WorldContext="WorldContextObject"))
+	UE_API FSubsonicEventInstance PlayEvent(const UObject* WorldContextObject, const FName EventName);
+
+	UFUNCTION(BlueprintCallable, Category = "Subsonic")
+	UE_API void StopEvent(FSubsonicEventInstance EventInstance);
+private:
+	TPimplPtr<FSubsonicSubsystemImpl> Impl;
+};
+
+#undef UE_API
diff --git a/Engine/Plugins/Experimental/Subsonic/Source/Subsonic/Public/SubsonicTypes.h b/Engine/Plugins/Experimental/Subsonic/Source/Subsonic/Public/SubsonicTypes.h
new file mode 100644
--- /dev/null
+++ b/Engine/Plugins/Experimental/Subsonic/Source/Subsonic/Public/SubsonicTypes.h
@@ -0,0 +1,27 @@
+// Copyright Epic Games, Inc. All Rights Reserved.
+
+#pragma once
+
+#include "CoreTypes.h"
+
+#include "SubsonicTypes.generated.h"
+
+/**
+* Handle to a playing instance of a Subsonic event.
+*/
+USTRUCT(BlueprintType)
+struct FSubsonicEventInstance final
+{
+	GENERATED_USTRUCT_BODY();
+
+	bool operator==(const FSubsonicEventInstance& Other) const
+	{
+		return RawHandle == Other.RawHandle;
+	}
+
+private:
+	uint32 RawHandle = 0;
+
+	friend class FSubsonicSubsystemImpl;
+	friend uint32 GetTypeHash(FSubsonicEventInstance Instance);
+};
diff --git a/Engine/Plugins/Experimental/Subsonic/Source/Subsonic/Subsonic.Build.cs b/Engine/Plugins/Experimental/Subsonic/Source/Subsonic/Subsonic.Build.cs
--- a/Engine/Plugins/Experimental/Subsonic/Source/Subsonic/Subsonic.Build.cs
+++ b/Engine/Plugins/Experimental/Subsonic/Source/Subsonic/Subsonic.Build.cs
@@ -18,6 +18,8 @@ public class Subsonic : ModuleRules
 		PrivateDependencyModuleNames.AddRange(
 			new string[]
 			{
+				"CoreUObject",
+				"Engine",
 			}
 		);
 	}
diff --git a/Engine/Plugins/Experimental/Subsonic/Source/SubsonicEditor/Private/SubsonicAssetDefinitions.cpp b/Engine/Plugins/Experimental/Subsonic/Source/SubsonicEditor/Private/SubsonicAssetDefinitions.cpp
new file mode 100644
--- /dev/null
+++ b/Engine/Plugins/Experimental/Subsonic/Source/SubsonicEditor/Private/SubsonicAssetDefinitions.cpp
@@ -0,0 +1,73 @@
+// Copyright Epic Games, Inc. All Rights Reserved.
+
+#include "SubsonicAssetDefinitions.h"
+
+#include "Styling/SlateStyleRegistry.h"
+#include "SubsonicEventHandler.h"
+#include "SubsonicEventHandlerEditor.h"
+
+#include UE_INLINE_GENERATED_CPP_BY_NAME(SubsonicAssetDefinitions)
+
+#define LOCTEXT_NAMESPACE "SubsonicEditor"
+
+
+FLinearColor UAssetDefinition_SubsonicEventHandler::GetAssetColor() const
+{
+ 	if (const ISlateStyle* SubsonicStyle = FSlateStyleRegistry::FindSlateStyle("SubsonicStyle"))
+ 	{
+ 		return SubsonicStyle->GetColor("SubsonicEventHandler.Color").ToFColorSRGB();
+ 	}
+ 
+ 	return FColor::White;
+}
+
+TSoftClassPtr<UObject> UAssetDefinition_SubsonicEventHandler::GetAssetClass() const
+{
+	return USubsonicEventHandler::StaticClass();
+}
+
+TConstArrayView<FAssetCategoryPath> UAssetDefinition_SubsonicEventHandler::GetAssetCategories() const
+{
+	static const TArray<FAssetCategoryPath> Categories = { EAssetCategoryPaths::Audio / LOCTEXT("AssetSoundSubsonicSubMenu", "Subsonic") };
+	return Categories;
+}
+
+EAssetCommandResult UAssetDefinition_SubsonicEventHandler::OpenAssets(const FAssetOpenArgs& OpenArgs) const
+{
+	using namespace UE::Subsonic::Editor;
+
+	for (const FAssetData& AssetData : OpenArgs.Assets)
+	{
+		TSharedRef<FEventHandlerEditor> NewEditor = MakeShared<FEventHandlerEditor>();
+		NewEditor->Init(OpenArgs.GetToolkitMode(), OpenArgs.ToolkitHost, *CastChecked<USubsonicEventHandler>(AssetData.GetAsset()));
+	}
+
+	return EAssetCommandResult::Handled;
+}
+
+const FSlateBrush* UAssetDefinition_SubsonicEventHandler::GetThumbnailBrush(const FAssetData& InAssetData, const FName InClassName) const
+{
+	return Super::GetThumbnailBrush(InAssetData, InClassName);
+}
+
+const FSlateBrush* UAssetDefinition_SubsonicEventHandler::GetIconBrush(const FAssetData& InAssetData, const FName InClassName) const
+{
+	return Super::GetIconBrush(InAssetData, InClassName);
+}
+
+bool UAssetDefinition_SubsonicEventHandler::GetThumbnailActionOverlay(const FAssetData& InAssetData, FAssetActionThumbnailOverlayInfo& OutActionOverlayInfo) const
+{
+	return Super::GetThumbnailActionOverlay(InAssetData, OutActionOverlayInfo);
+}
+
+EAssetCommandResult UAssetDefinition_SubsonicEventHandler::ActivateAssets(const FAssetActivateArgs& ActivateArgs) const
+{
+	return Super::ActivateAssets(ActivateArgs);
+}
+
+void UAssetDefinition_SubsonicEventHandler::GetAssetActionButtonExtensions(const FAssetData& InAssetData, TArray<FAssetButtonActionExtension>& OutExtensions) const
+{
+	Super::GetAssetActionButtonExtensions(InAssetData, OutExtensions);
+}
+
+#undef LOCTEXT_NAMESPACE // SubsonicEditor
diff --git a/Engine/Plugins/Experimental/Subsonic/Source/SubsonicEditor/Private/SubsonicAssetDefinitions.h b/Engine/Plugins/Experimental/Subsonic/Source/SubsonicEditor/Private/SubsonicAssetDefinitions.h
new file mode 100644
--- /dev/null
+++ b/Engine/Plugins/Experimental/Subsonic/Source/SubsonicEditor/Private/SubsonicAssetDefinitions.h
@@ -0,0 +1,29 @@
+// Copyright Epic Games, Inc. All Rights Reserved.
+#pragma once
+
+#include "AssetTypeActions/AssetDefinitionDefault_AudioDiffable.h"
+
+#include "SubsonicAssetDefinitions.generated.h"
+
+
+UCLASS()
+class UAssetDefinition_SubsonicEventHandler : public UAssetDefinitionDefault_AudioDiffable
+{
+	GENERATED_BODY()
+
+public:
+	// UAssetDefinition Begin
+	virtual FText GetAssetDisplayName() const override { return NSLOCTEXT("AssetTypeActions", "AssetTypeActions_SubsonicEventHandler", "Subsonic Event Handler"); }
+	virtual FLinearColor GetAssetColor() const override;
+	virtual TSoftClassPtr<UObject> GetAssetClass() const override;
+	virtual TConstArrayView<FAssetCategoryPath> GetAssetCategories() const override;
+	virtual EAssetCommandResult OpenAssets(const FAssetOpenArgs& OpenArgs) const override;
+
+	virtual bool GetThumbnailActionOverlay(const FAssetData& InAssetData, FAssetActionThumbnailOverlayInfo& OutActionOverlayInfo) const override;
+	virtual EAssetCommandResult ActivateAssets(const FAssetActivateArgs& ActivateArgs) const override;
+	virtual void GetAssetActionButtonExtensions(const FAssetData& InAssetData, TArray<FAssetButtonActionExtension>& OutExtensions) const override;
+
+ 	virtual const FSlateBrush* GetThumbnailBrush(const FAssetData& InAssetData, const FName InClassName) const override;
+ 	virtual const FSlateBrush* GetIconBrush(const FAssetData& InAssetData, const FName InClassName) const override;
+	// UAssetDefinition End
+};
\ No newline at end of file
diff --git a/Engine/Plugins/Experimental/Subsonic/Source/SubsonicEditor/Private/SubsonicEditorModule.cpp b/Engine/Plugins/Experimental/Subsonic/Source/SubsonicEditor/Private/SubsonicEditorModule.cpp
--- a/Engine/Plugins/Experimental/Subsonic/Source/SubsonicEditor/Private/SubsonicEditorModule.cpp
+++ b/Engine/Plugins/Experimental/Subsonic/Source/SubsonicEditor/Private/SubsonicEditorModule.cpp
@@ -1,5 +1,27 @@
 // Copyright Epic Games, Inc. All Rights Reserved.
 
+#include "SubsonicSlateStyle.h"
 #include "Modules/ModuleManager.h"
 
-IMPLEMENT_MODULE(FDefaultModuleImpl, SubsonicEditor)
+
+namespace UE::Subsonic::Editor
+{
+	class FModule : public IModuleInterface
+	{
+	public:
+		virtual void StartupModule() override
+		{
+			StyleSet = MakeShared<FSlateStyle>();
+		}
+
+		virtual void ShutdownModule() override
+		{
+			StyleSet.Reset();
+		}
+
+		TSharedPtr<FSlateStyleSet> StyleSet;
+	};
+} // namespace UE::Subsonic::Editor
+
+
+IMPLEMENT_MODULE(UE::Subsonic::Editor::FModule, SubsonicEditor);
diff --git a/Engine/Plugins/Experimental/Subsonic/Source/SubsonicEditor/Private/SubsonicEventHandlerEditor.cpp b/Engine/Plugins/Experimental/Subsonic/Source/SubsonicEditor/Private/SubsonicEventHandlerEditor.cpp
new file mode 100644
--- /dev/null
+++ b/Engine/Plugins/Experimental/Subsonic/Source/SubsonicEditor/Private/SubsonicEventHandlerEditor.cpp
@@ -0,0 +1,198 @@
+// Copyright Epic Games, Inc. All Rights Reserved.
+
+#include "SubsonicEventHandlerEditor.h"
+
+#include "Framework/Application/SlateApplication.h"
+#include "Framework/Docking/TabManager.h"
+#include "PropertyEditorModule.h"
+#include "Modules/ModuleManager.h"
+#include "Widgets/Docking/SDockTab.h"
+#include "Styling/SlateStyleRegistry.h"
+#include "SubsonicEventHandler.h"
+
+#define LOCTEXT_NAMESPACE "SubsonicEditor"
+
+
+namespace UE::Subsonic::Editor
+{
+	namespace EventHandlerEditorPrivate
+	{
+		const FName TabName_Details = "SubsonicEventHandlerEditor_Details";
+	} // namespace EventHandlerPrivate
+
+	const FName FEventHandlerEditor::EditorName = "SubsonicEventHandlerEditor";
+
+	void FEventHandlerEditor::RegisterTabSpawners(const TSharedRef<FTabManager>& InTabManager)
+	{
+		WorkspaceMenuCategory = InTabManager->AddLocalWorkspaceMenuCategory(LOCTEXT("WorkspaceMenu_SubsonicEventHandlerEditor", "Subsonic EventHandler Editor"));
+
+		FDetailsViewArgs Args;
+		Args.bHideSelectionTip = true;
+		Args.NotifyHook = this;
+
+		FPropertyEditorModule& PropertyModule = FModuleManager::LoadModuleChecked<FPropertyEditorModule>("PropertyEditor");
+		TSharedRef<IDetailsView> NewDetailsView = PropertyModule.CreateDetailView(Args);
+
+		InTabManager->RegisterTabSpawner(EventHandlerEditorPrivate::TabName_Details, FOnSpawnTab::CreateLambda([NewDetailsView](const FSpawnTabArgs& Args)
+			{
+				return SNew(SDockTab).Label(LOCTEXT("SubsonicEventHandlerDetailsTitle", "Details"))[ NewDetailsView ];
+			}))
+			.SetDisplayName(LOCTEXT("DetailsTab", "Details"))
+			.SetGroup(WorkspaceMenuCategory.ToSharedRef())
+			.SetIcon(FSlateIcon(FAppStyle::GetAppStyleSetName(), "LevelEditor.Tabs.Details"));
+
+		DetailsView = NewDetailsView;
+
+		FAssetEditorToolkit::RegisterTabSpawners(InTabManager);
+	}
+
+	void FEventHandlerEditor::UnregisterTabSpawners(const TSharedRef<FTabManager>& InTabManager)
+	{
+		FAssetEditorToolkit::UnregisterTabSpawners(InTabManager);
+
+		InTabManager->UnregisterTabSpawner(EventHandlerEditorPrivate::TabName_Details);
+	}
+
+	FEventHandlerEditor::~FEventHandlerEditor()
+	{
+		if (GEditor)
+		{
+			GEditor->UnregisterForUndo(this);
+		}
+	}
+
+	void FEventHandlerEditor::Init(const EToolkitMode::Type Mode, const TSharedPtr<IToolkitHost>& InitToolkitHost, USubsonicEventHandler& InEventHandler)
+	{
+		// Support undo/redo
+		InEventHandler.SetFlags(RF_Transactional);
+
+		if (GEditor)
+		{
+			GEditor->RegisterForUndo(this);
+		}
+
+		DetailsStack = FTabManager::NewStack()
+			->SetSizeCoefficient(0.50f)
+			->SetHideTabWell(false)
+			->AddTab(EventHandlerEditorPrivate::TabName_Details, ETabState::OpenedTab);
+
+		const TSharedRef<FTabManager::FLayout> StandaloneDefaultLayout = FTabManager::NewLayout("Standalone_SubsonicEventHandlerEditor_Layout_v1")
+			->AddArea
+			(
+				FTabManager::NewPrimaryArea()
+				->SetOrientation(Orient_Vertical)
+				->Split(FTabManager::NewSplitter()
+					->SetOrientation(Orient_Horizontal)
+					->Split
+					(
+						FTabManager::NewSplitter()
+						->SetSizeCoefficient(0.15f)
+						->SetOrientation(Orient_Vertical)
+						->Split
+						(
+							FTabManager::NewStack()
+							->SetSizeCoefficient(0.25f)
+							->SetHideTabWell(false)
+							->AddTab(EventHandlerEditorPrivate::TabName_Details, ETabState::OpenedTab)
+						)
+					)
+				)
+			);
+
+		{
+			constexpr bool bCreateDefaultStandaloneMenu = true;
+			constexpr bool bCreateDefaultToolbar = true;
+			constexpr bool bToolbarFocusable = false;
+			constexpr bool bUseSmallToolbarIcons = true;
+
+			FAssetEditorToolkit::InitAssetEditor(
+				Mode,
+				InitToolkitHost,
+				FName("SubsonicEventHandlerEditorApp"),
+				StandaloneDefaultLayout,
+				bCreateDefaultStandaloneMenu,
+				bCreateDefaultToolbar,
+				&InEventHandler,
+				bToolbarFocusable,
+				bUseSmallToolbarIcons);
+		}
+
+		if (DetailsView.IsValid())
+		{
+			DetailsView->SetObject(&InEventHandler);
+		}
+
+		RegenerateMenusAndToolbars();
+	}
+
+	USubsonicEventHandler* FEventHandlerEditor::GetEventHandler() const
+	{
+		return Cast<USubsonicEventHandler>(GetEditingObject());
+	}
+
+	FName FEventHandlerEditor::GetToolkitFName() const
+	{
+		return FEventHandlerEditor::EditorName;
+	}
+
+	FText FEventHandlerEditor::GetBaseToolkitName() const
+	{
+		return LOCTEXT("AppLabel", "Subsonic EventHandler Editor");
+	}
+
+	FString FEventHandlerEditor::GetWorldCentricTabPrefix() const
+	{
+		return LOCTEXT("WorldCentricTabPrefix", "Subsonic Event Handler ").ToString();
+	}
+
+	FLinearColor FEventHandlerEditor::GetWorldCentricTabColorScale() const
+	{
+		if (const ISlateStyle* Style = FSlateStyleRegistry::FindSlateStyle("SubsonicStyle"))
+		{
+			return Style->GetColor("SubsonicEventHandler.Color");
+		}
+
+		return FLinearColor(0.3f, 0.2f, 0.5f, 0.5f);
+	}
+
+	const FSlateBrush* FEventHandlerEditor::GetDefaultTabIcon() const
+	{
+		if (const ISlateStyle* Style = FSlateStyleRegistry::FindSlateStyle("SubsonicStyle"))
+		{
+			if (const FSlateBrush* Brush = Style->GetBrush("SubsonicEventHandlerEditor_Icon"))
+			{
+				return Brush;
+			}
+		}
+
+		if (const FSlateBrush* NoBrush = FAppStyle::GetBrush("NoBrush"))
+		{
+			return NoBrush;
+		}
+
+		static const FSlateBrush NullBrush;
+		return &NullBrush;
+	}
+
+	FLinearColor FEventHandlerEditor::GetDefaultTabColor() const
+	{
+		if (const ISlateStyle* Style = FSlateStyleRegistry::FindSlateStyle("SubsonicStyle"))
+		{
+			return Style->GetColor("SubsonicEventHandler.Color");
+		}
+
+		return FAssetEditorToolkit::GetDefaultTabColor();
+	}
+
+	FName FEventHandlerEditor::GetEditorName() const
+	{
+		return FEventHandlerEditor::EditorName;
+	}
+
+	void FEventHandlerEditor::PostUndo(bool bSuccess)
+	{
+		FSlateApplication::Get().DismissAllMenus();
+		FEditorUndoClient::PostUndo(bSuccess);
+	}
+}
+#undef LOCTEXT_NAMESPACE // SubsonicEditor
diff --git a/Engine/Plugins/Experimental/Subsonic/Source/SubsonicEditor/Private/SubsonicEventHandlerEditor.h b/Engine/Plugins/Experimental/Subsonic/Source/SubsonicEditor/Private/SubsonicEventHandlerEditor.h
new file mode 100644
--- /dev/null
+++ b/Engine/Plugins/Experimental/Subsonic/Source/SubsonicEditor/Private/SubsonicEventHandlerEditor.h
@@ -0,0 +1,55 @@
+// Copyright Epic Games, Inc. All Rights Reserved.
+#pragma once
+
+#include "SubsonicEventHandlerEditor.h"
+
+#include "EditorUndoClient.h"
+#include "Framework/Docking/TabManager.h"
+#include "IDetailsView.h"
+#include "Misc/NotifyHook.h"
+#include "Templates/SharedPointer.h"
+#include "Toolkits/AssetEditorToolkit.h"
+#include "Toolkits/IToolkitHost.h"
+
+
+// Forward Declarations
+class USubsonicEventHandler;
+
+
+namespace UE::Subsonic::Editor
+{
+	class FEventHandlerEditor : public FAssetEditorToolkit, public FNotifyHook, public FEditorUndoClient
+	{
+	public:
+		static const FName EditorName;
+
+		FEventHandlerEditor() = default;
+		virtual ~FEventHandlerEditor();
+
+		virtual void RegisterTabSpawners(const TSharedRef<FTabManager>& TabManager) override;
+		virtual void UnregisterTabSpawners(const TSharedRef<FTabManager>& TabManager) override;
+
+		void Init(const EToolkitMode::Type Mode, const TSharedPtr<IToolkitHost>& InitToolkitHost, USubsonicEventHandler& EventHandler);
+
+		USubsonicEventHandler* GetEventHandler() const;
+
+		/** IToolkit interface */
+		virtual FName GetToolkitFName() const override;
+		virtual FText GetBaseToolkitName() const override;
+		virtual FString GetWorldCentricTabPrefix() const override;
+		virtual FLinearColor GetWorldCentricTabColorScale() const override;
+		virtual const FSlateBrush* GetDefaultTabIcon() const override;
+		virtual FLinearColor GetDefaultTabColor() const override;
+
+		/** IAssetEditorInstance interface */
+		virtual FName GetEditorName() const override;
+
+		/** FEditorUndoClient Interface */
+		virtual void PostUndo(bool bSuccess) override;
+		virtual void PostRedo(bool bSuccess) override { PostUndo(bSuccess); }
+
+	private:
+		TSharedPtr<IDetailsView> DetailsView;
+		TSharedPtr<FTabManager::FStack> DetailsStack;
+	};
+} // namespace UE::Subsonic::Editor
diff --git a/Engine/Plugins/Experimental/Subsonic/Source/SubsonicEditor/Private/SubsonicEventHandlerFactory.cpp b/Engine/Plugins/Experimental/Subsonic/Source/SubsonicEditor/Private/SubsonicEventHandlerFactory.cpp
new file mode 100644
--- /dev/null
+++ b/Engine/Plugins/Experimental/Subsonic/Source/SubsonicEditor/Private/SubsonicEventHandlerFactory.cpp
@@ -0,0 +1,22 @@
+// Copyright Epic Games, Inc. All Rights Reserved.
+#include "SubsonicEventHandlerFactory.h"
+
+#include "SubsonicEventHandler.h"
+
+#include UE_INLINE_GENERATED_CPP_BY_NAME(SubsonicEventHandlerFactory)
+
+
+USubsonicEventHandlerFactory::USubsonicEventHandlerFactory(const FObjectInitializer& ObjectInitializer)
+	: Super(ObjectInitializer)
+{
+	bCreateNew = true;
+	bEditorImport = false;
+	bEditAfterNew = true;
+
+	SupportedClass = USubsonicEventHandler::StaticClass();
+}
+
+UObject* USubsonicEventHandlerFactory::FactoryCreateNew(UClass* InClass, UObject* InParent, FName InName, EObjectFlags InFlags, UObject* InContext, FFeedbackContext* InFeedbackContext)
+{
+	return NewObject<USubsonicEventHandler>(InParent, InName, InFlags);
+}
diff --git a/Engine/Plugins/Experimental/Subsonic/Source/SubsonicEditor/Private/SubsonicEventHandlerFactory.h b/Engine/Plugins/Experimental/Subsonic/Source/SubsonicEditor/Private/SubsonicEventHandlerFactory.h
new file mode 100644
--- /dev/null
+++ b/Engine/Plugins/Experimental/Subsonic/Source/SubsonicEditor/Private/SubsonicEventHandlerFactory.h
@@ -0,0 +1,19 @@
+// Copyright Epic Games, Inc. All Rights Reserved.
+
+#pragma once
+
+#include "Factories/Factory.h"
+#include "UObject/Object.h"
+
+#include "SubsonicEventHandlerFactory.generated.h"
+
+
+UCLASS()
+class USubsonicEventHandlerFactory : public UFactory
+{
+	GENERATED_UCLASS_BODY()
+
+	//~ Begin UFactory Interface
+	virtual UObject* FactoryCreateNew(UClass* InClass, UObject* InParent, FName InName, EObjectFlags InFlags, UObject* InContext, FFeedbackContext* InFeedbackContext) override;
+	//~ Begin UFactory Interface
+};
diff --git a/Engine/Plugins/Experimental/Subsonic/Source/SubsonicEditor/Private/SubsonicSlateStyle.h b/Engine/Plugins/Experimental/Subsonic/Source/SubsonicEditor/Private/SubsonicSlateStyle.h
new file mode 100644
--- /dev/null
+++ b/Engine/Plugins/Experimental/Subsonic/Source/SubsonicEditor/Private/SubsonicSlateStyle.h
@@ -0,0 +1,42 @@
+// Copyright Epic Games, Inc. All Rights Reserved.
+#pragma once
+
+#include "AudioWidgetsStyle.h"
+#include "Styling/SlateStyleRegistry.h"
+
+
+namespace UE::Subsonic::Editor
+{
+	class FSlateStyle final : public FSlateStyleSet
+	{
+	public:
+		FSlateStyle()
+			: FSlateStyleSet("SubsonicStyle")
+		{
+			SetParentStyleName(FAudioWidgetsStyle::Get().GetStyleSetName());
+
+			SetContentRoot(FPaths::EnginePluginsDir() / TEXT("Runtime/Subsonic/Content/Editor/Slate"));
+			SetCoreContentRoot(FPaths::EngineContentDir() / TEXT("Slate"));
+
+			static const FVector2D Icon20x20(20.0f, 20.0f);
+			static const FVector2D Icon40x40(40.0f, 40.0f);
+
+			static const FVector2D Icon16 = FVector2D(16.0f, 16.0f);
+			static const FVector2D Icon64 = FVector2D(64.0f, 64.0f);
+
+			const FVector2D Icon15x11(15.0f, 11.0f);
+
+			// Subsonic Editor
+			{
+				Set("SubsonicEventHandler.Color", FColor(156, 0, 255));
+			}
+
+			FSlateStyleRegistry::RegisterSlateStyle(*this);
+		}
+
+		~FSlateStyle()
+		{
+			FSlateStyleRegistry::UnRegisterSlateStyle(*this);
+		}
+	};
+} // namespace UE::Subsonic::Editor
\ No newline at end of file
diff --git a/Engine/Plugins/Experimental/Subsonic/Source/SubsonicEditor/SubsonicEditor.Build.cs b/Engine/Plugins/Experimental/Subsonic/Source/SubsonicEditor/SubsonicEditor.Build.cs
--- a/Engine/Plugins/Experimental/Subsonic/Source/SubsonicEditor/SubsonicEditor.Build.cs
+++ b/Engine/Plugins/Experimental/Subsonic/Source/SubsonicEditor/SubsonicEditor.Build.cs
@@ -12,12 +12,22 @@ public class SubsonicEditor : ModuleRules
 			new string[]
 			{
 				"Core",
+				"CoreUObject"
 			}
 		);
 			
 		PrivateDependencyModuleNames.AddRange(
 			new string[]
 			{
+				"AssetDefinition",
+				"AudioEditor",
+				"AudioWidgets",
+				"ContentBrowser",
+				"PropertyEditor",
+				"Slate",
+				"SlateCore",
+				"Subsonic",
+				"UnrealEd"
 			}
 		);
 	}
