[Backout] - CL49867385
#rnx
Original CL Desc
-----------------------------------------------------------------
Checkpoint (Combined Team Work):
- Subsonic Event Handler Asset, Factory, Editor, & Definition stubs
- Add Subsonic plugin to applicable test projects dev/testing

#rnx

#Synced-CL 49790744

diff --git a/Engine/Plugins/Experimental/Subsonic/Source/Subsonic/Private/SubsonicEventHandler.cpp b/Engine/Plugins/Experimental/Subsonic/Source/Subsonic/Private/SubsonicEventHandler.cpp
deleted file mode 100644
diff --git a/Engine/Plugins/Experimental/Subsonic/Source/Subsonic/Private/SubsonicEventInstanceImpl.cpp b/Engine/Plugins/Experimental/Subsonic/Source/Subsonic/Private/SubsonicEventInstanceImpl.cpp
deleted file mode 100644
diff --git a/Engine/Plugins/Experimental/Subsonic/Source/Subsonic/Private/SubsonicEventInstanceImpl.h b/Engine/Plugins/Experimental/Subsonic/Source/Subsonic/Private/SubsonicEventInstanceImpl.h
deleted file mode 100644
diff --git a/Engine/Plugins/Experimental/Subsonic/Source/Subsonic/Private/SubsonicSubsystem.cpp b/Engine/Plugins/Experimental/Subsonic/Source/Subsonic/Private/SubsonicSubsystem.cpp
deleted file mode 100644
diff --git a/Engine/Plugins/Experimental/Subsonic/Source/Subsonic/Private/SubsonicTempDataAsset.h b/Engine/Plugins/Experimental/Subsonic/Source/Subsonic/Private/SubsonicTempDataAsset.h
deleted file mode 100644
diff --git a/Engine/Plugins/Experimental/Subsonic/Source/Subsonic/Private/SubsonicTypes.cpp b/Engine/Plugins/Experimental/Subsonic/Source/Subsonic/Private/SubsonicTypes.cpp
deleted file mode 100644
diff --git a/Engine/Plugins/Experimental/Subsonic/Source/Subsonic/Public/SubsonicEventHandler.h b/Engine/Plugins/Experimental/Subsonic/Source/Subsonic/Public/SubsonicEventHandler.h
deleted file mode 100644
diff --git a/Engine/Plugins/Experimental/Subsonic/Source/Subsonic/Public/SubsonicSubsystem.h b/Engine/Plugins/Experimental/Subsonic/Source/Subsonic/Public/SubsonicSubsystem.h
deleted file mode 100644
diff --git a/Engine/Plugins/Experimental/Subsonic/Source/Subsonic/Public/SubsonicTypes.h b/Engine/Plugins/Experimental/Subsonic/Source/Subsonic/Public/SubsonicTypes.h
deleted file mode 100644
diff --git a/Engine/Plugins/Experimental/Subsonic/Source/Subsonic/Subsonic.Build.cs b/Engine/Plugins/Experimental/Subsonic/Source/Subsonic/Subsonic.Build.cs
--- a/Engine/Plugins/Experimental/Subsonic/Source/Subsonic/Subsonic.Build.cs
+++ b/Engine/Plugins/Experimental/Subsonic/Source/Subsonic/Subsonic.Build.cs
@@ -18,8 +18,6 @@ public class Subsonic : ModuleRules
 		PrivateDependencyModuleNames.AddRange(
 			new string[]
 			{
-				"CoreUObject",
-				"Engine",
 			}
 		);
 	}
diff --git a/Engine/Plugins/Experimental/Subsonic/Source/SubsonicEditor/Private/SubsonicAssetDefinitions.cpp b/Engine/Plugins/Experimental/Subsonic/Source/SubsonicEditor/Private/SubsonicAssetDefinitions.cpp
deleted file mode 100644
diff --git a/Engine/Plugins/Experimental/Subsonic/Source/SubsonicEditor/Private/SubsonicAssetDefinitions.h b/Engine/Plugins/Experimental/Subsonic/Source/SubsonicEditor/Private/SubsonicAssetDefinitions.h
deleted file mode 100644
diff --git a/Engine/Plugins/Experimental/Subsonic/Source/SubsonicEditor/Private/SubsonicEditorModule.cpp b/Engine/Plugins/Experimental/Subsonic/Source/SubsonicEditor/Private/SubsonicEditorModule.cpp
--- a/Engine/Plugins/Experimental/Subsonic/Source/SubsonicEditor/Private/SubsonicEditorModule.cpp
+++ b/Engine/Plugins/Experimental/Subsonic/Source/SubsonicEditor/Private/SubsonicEditorModule.cpp
@@ -1,27 +1,5 @@
 // Copyright Epic Games, Inc. All Rights Reserved.
 
-#include "SubsonicSlateStyle.h"
 #include "Modules/ModuleManager.h"
 
-
-namespace UE::Subsonic::Editor
-{
-	class FModule : public IModuleInterface
-	{
-	public:
-		virtual void StartupModule() override
-		{
-			StyleSet = MakeShared<FSlateStyle>();
-		}
-
-		virtual void ShutdownModule() override
-		{
-			StyleSet.Reset();
-		}
-
-		TSharedPtr<FSlateStyleSet> StyleSet;
-	};
-} // namespace UE::Subsonic::Editor
-
-
-IMPLEMENT_MODULE(UE::Subsonic::Editor::FModule, SubsonicEditor);
+IMPLEMENT_MODULE(FDefaultModuleImpl, SubsonicEditor)
diff --git a/Engine/Plugins/Experimental/Subsonic/Source/SubsonicEditor/Private/SubsonicEventHandlerEditor.cpp b/Engine/Plugins/Experimental/Subsonic/Source/SubsonicEditor/Private/SubsonicEventHandlerEditor.cpp
deleted file mode 100644
diff --git a/Engine/Plugins/Experimental/Subsonic/Source/SubsonicEditor/Private/SubsonicEventHandlerEditor.h b/Engine/Plugins/Experimental/Subsonic/Source/SubsonicEditor/Private/SubsonicEventHandlerEditor.h
deleted file mode 100644
diff --git a/Engine/Plugins/Experimental/Subsonic/Source/SubsonicEditor/Private/SubsonicEventHandlerFactory.cpp b/Engine/Plugins/Experimental/Subsonic/Source/SubsonicEditor/Private/SubsonicEventHandlerFactory.cpp
deleted file mode 100644
diff --git a/Engine/Plugins/Experimental/Subsonic/Source/SubsonicEditor/Private/SubsonicEventHandlerFactory.h b/Engine/Plugins/Experimental/Subsonic/Source/SubsonicEditor/Private/SubsonicEventHandlerFactory.h
deleted file mode 100644
diff --git a/Engine/Plugins/Experimental/Subsonic/Source/SubsonicEditor/Private/SubsonicSlateStyle.h b/Engine/Plugins/Experimental/Subsonic/Source/SubsonicEditor/Private/SubsonicSlateStyle.h
deleted file mode 100644
diff --git a/Engine/Plugins/Experimental/Subsonic/Source/SubsonicEditor/SubsonicEditor.Build.cs b/Engine/Plugins/Experimental/Subsonic/Source/SubsonicEditor/SubsonicEditor.Build.cs
--- a/Engine/Plugins/Experimental/Subsonic/Source/SubsonicEditor/SubsonicEditor.Build.cs
+++ b/Engine/Plugins/Experimental/Subsonic/Source/SubsonicEditor/SubsonicEditor.Build.cs
@@ -12,22 +12,12 @@ public class SubsonicEditor : ModuleRules
 			new string[]
 			{
 				"Core",
-				"CoreUObject"
 			}
 		);
 			
 		PrivateDependencyModuleNames.AddRange(
 			new string[]
 			{
-				"AssetDefinition",
-				"AudioEditor",
-				"AudioWidgets",
-				"ContentBrowser",
-				"PropertyEditor",
-				"Slate",
-				"SlateCore",
-				"Subsonic",
-				"UnrealEd"
 			}
 		);
 	}
