Subsonic Collection Checkpoint
- Type rename (Handler -> Collection)
- Named events to Gameplay Tags
- Basic AudioComponent action implementation with accessor and subaction examples

#rnx

diff --git a/Engine/Plugins/Experimental/Subsonic/Source/Subsonic/Private/CoreActions/SubsonicAction_AudioComponent.cpp b/Engine/Plugins/Experimental/Subsonic/Source/Subsonic/Private/CoreActions/SubsonicAction_AudioComponent.cpp
new file mode 100644
--- /dev/null
+++ b/Engine/Plugins/Experimental/Subsonic/Source/Subsonic/Private/CoreActions/SubsonicAction_AudioComponent.cpp
@@ -0,0 +1,183 @@
+// Copyright Epic Games, Inc. All Rights Reserved.
+
+#include "CoreActions/SubsonicAction_AudioComponent.h"
+#include "IAudioModulation.h"
+#include "SubsonicHandles.h"
+
+
+namespace UE::Subsonic
+{
+	namespace AudioComponentPrivate
+	{
+		UAudioComponent* AddScopedComponent(const UE::Subsonic::FActionHandle& InHandle, ESubsonicExecutionScope Scope, FName Name)
+		{
+			UAudioComponent* Component = nullptr;
+			if (USubsonicAudioComponentSubsystem* AudioComponentSubsystem = InHandle.Event.Collection.FindDeviceSubsystem<USubsonicAudioComponentSubsystem>())
+			{
+				switch (Scope)
+				{
+					case ESubsonicExecutionScope::Global:
+					{
+						Component = AudioComponentSubsystem->AddComponent(Name);
+					}
+					break;
+
+					case ESubsonicExecutionScope::Collection:
+					{
+						Component = AudioComponentSubsystem->AddComponent(InHandle.Event.Collection, Name);
+					}
+					break;
+
+					case ESubsonicExecutionScope::Event:
+					{
+						Component = AudioComponentSubsystem->AddComponent(InHandle.Event, Name);
+					}
+					break;
+
+					case ESubsonicExecutionScope::Action:
+					{
+						Component = AudioComponentSubsystem->AddComponent(InHandle, Name);
+					}
+					break;
+				}
+			}
+
+			return Component;
+		}
+
+		UAudioComponent* FindScopedComponent(const UE::Subsonic::FActionHandle& InHandle, ESubsonicExecutionScope Scope, FName Name)
+		{
+			using namespace UE::Subsonic;
+
+			UAudioComponent* Component = nullptr;
+			if (USubsonicAudioComponentSubsystem* AudioComponentSubsystem = InHandle.Event.Collection.FindDeviceSubsystem<USubsonicAudioComponentSubsystem>())
+			{
+				switch (Scope)
+					{
+					case ESubsonicExecutionScope::Global:
+					{
+						Component = AudioComponentSubsystem->FindComponent(Name);
+					}
+					break;
+
+					case ESubsonicExecutionScope::Collection:
+					{
+						Component = AudioComponentSubsystem->FindComponent(InHandle.Event.Collection, Name);
+					}
+					break;
+
+					case ESubsonicExecutionScope::Event:
+					{
+						Component = AudioComponentSubsystem->FindComponent(InHandle.Event, Name);
+					}
+					break;
+
+					case ESubsonicExecutionScope::Action:
+					{
+						Component = AudioComponentSubsystem->FindComponent(InHandle, Name);
+					}
+					break;
+				}
+			}
+
+			return Component;
+		}
+
+		UAudioComponent* AccessComponent(const UE::Subsonic::FActionHandle& InHandle, ESubsonicAudioComponentAccess Access, ESubsonicExecutionScope Scope, FName Name)
+		{
+			UAudioComponent* Component = nullptr;
+			if (Access == ESubsonicAudioComponentAccess::Find || Access == ESubsonicAudioComponentAccess::FindOrAdd)
+			{
+				Component = FindScopedComponent(InHandle, Scope, Name);
+			}
+
+			bool bAddNewComponent = Access == ESubsonicAudioComponentAccess::Add;
+			if (Access == ESubsonicAudioComponentAccess::FindOrAdd)
+			{
+				// Didn't find an existing component above, so attempt
+				// to make one if set to find or add even if access not set to 'Add'.
+				if (!Component)
+				{
+					bAddNewComponent = true;
+				}
+			}
+
+			if (bAddNewComponent)
+			{
+				Component = AddScopedComponent(InHandle, Scope, Name);
+			}
+
+			return Component;
+		}
+	} // namespace AudioComponentPrivate
+}
+
+void FSubsonicEventAction_AudioComponentModify::Execute(const UE::Subsonic::FActionHandle& InHandle) const
+{
+	using namespace UE::Subsonic;
+
+	if (UAudioComponent* AudioComponent = AudioComponentPrivate::AccessComponent(InHandle, Access, Scope, Name))
+	{
+		for (const TInstancedStruct<FSubsonicEventAction_AudioComponentModifierBase>& Modifier : Modifiers)
+		{
+			if (const FSubsonicEventAction_AudioComponentModifierBase* ModPtr = Modifier.GetPtr())
+			{
+				ModPtr->Execute(InHandle, *AudioComponent);
+			}
+		}
+	}
+}
+
+void FSubsonicEventAction_AudioComponentPlay::Execute(const UE::Subsonic::FActionHandle& InHandle) const
+{
+	using namespace UE::Subsonic;
+
+	if (UAudioComponent* AudioComponent = AudioComponentPrivate::AccessComponent(InHandle, Access, Scope, Name))
+	{
+		AudioComponent->SetSound(Sound);
+		AudioComponent->Play();
+	}
+}
+
+void FSubsonicEventAction_AudioComponentStop::Execute(const UE::Subsonic::FActionHandle& InHandle) const
+{
+	using namespace UE::Subsonic;
+
+	if (UAudioComponent* AudioComponent = AudioComponentPrivate::FindScopedComponent(InHandle, Scope, Name))
+	{
+		AudioComponent->Stop();
+	}
+}
+
+void FSubsonicEventAction_AudioComponentModifier_Play::Execute(const UE::Subsonic::FActionHandle& InHandle, UAudioComponent& AudioComponent) const
+{
+	AudioComponent.Play(StartTime);
+}
+
+
+void FSubsonicEventAction_AudioComponentModifier_SetAttenuation::Execute(const UE::Subsonic::FActionHandle& InHandle, UAudioComponent& AudioComponent) const
+{
+	AudioComponent.SetAttenuationSettings(Attenuation);
+}
+
+void FSubsonicEventAction_AudioComponentModifier_SetConcurrency::Execute(const UE::Subsonic::FActionHandle& InHandle, UAudioComponent& AudioComponent) const
+{
+	AudioComponent.ConcurrencySet = Concurrency;
+}
+
+void FSubsonicEventAction_AudioComponentModifier_SetModulationRouting::Execute(const UE::Subsonic::FActionHandle& InHandle, UAudioComponent& AudioComponent) const
+{
+	TSet<USoundModulatorBase*> ModulatorSet;
+	Algo::Transform(Modulators, ModulatorSet, [](const TObjectPtr<USoundModulatorBase>& Modulator) { return Modulator.Get(); });
+	AudioComponent.SetModulationRouting(ModulatorSet, Destination, RoutingMethod);
+}
+
+void FSubsonicEventAction_AudioComponentModifier_SetSound::Execute(const UE::Subsonic::FActionHandle& InHandle, UAudioComponent& AudioComponent) const
+{
+	AudioComponent.SetSound(Sound);
+}
+
+void FSubsonicEventAction_AudioComponentModifier_Stop::Execute(const UE::Subsonic::FActionHandle& InHandle, UAudioComponent& AudioComponent) const
+{
+	AudioComponent.Stop();
+}
\ No newline at end of file
diff --git a/Engine/Plugins/Experimental/Subsonic/Source/Subsonic/Private/CoreActions/SubsonicAction_AudioComponent.h b/Engine/Plugins/Experimental/Subsonic/Source/Subsonic/Private/CoreActions/SubsonicAction_AudioComponent.h
new file mode 100644
--- /dev/null
+++ b/Engine/Plugins/Experimental/Subsonic/Source/Subsonic/Private/CoreActions/SubsonicAction_AudioComponent.h
@@ -0,0 +1,211 @@
+// Copyright Epic Games, Inc. All Rights Reserved.
+#pragma once
+
+#include "CoreEventSystems/SubsonicAudioComponentSubsystem.h"
+#include "Misc/CoreMiscDefines.h"
+#include "Sound/SoundBase.h"
+#include "Sound/SoundConcurrency.h"
+#include "Sound/SoundModulationDestination.h"
+#include "StructUtils/InstancedStruct.h"
+#include "SubsonicEventCollection.h"
+
+#include "SubsonicAction_AudioComponent.generated.h"
+
+
+// Forward Declarations
+class UAudioComponent;
+class USoundConcurrency;
+class USoundModulatorBase;
+
+struct FSubsonicEventAction_AudioComponentModifierBase;
+
+
+UENUM(BlueprintType)
+enum class ESubsonicAudioComponentAccess : uint8
+{
+	// Adds new component, releasing any existing component.
+	// Stops existing component if sound is not one-shot.
+	Add,
+
+	// Finds existing component, or creates a new one.
+	FindOrAdd,
+
+	// Ignore request and if component doesn't already exist.
+	Find
+};
+
+USTRUCT(MinimalAPI, BlueprintType, DisplayName = "Modify Audio Component")
+struct FSubsonicEventAction_AudioComponentModify : public FSubsonicEventActionBase
+{
+	GENERATED_BODY()
+
+public:
+	virtual ~FSubsonicEventAction_AudioComponentModify() = default;
+
+	virtual void Execute(const UE::Subsonic::FActionHandle& InHandle) const override;
+
+	// Addressable name of audio component instance.
+	UPROPERTY(EditAnywhere, Category = "Subsonic")
+	FName Name = "Sound01";
+
+	// Scope of component accessor.  If set to event, provided name is used
+	// to address the same component within the event, collection the same component
+	// within the collection, etc.
+	UPROPERTY(EditAnywhere, Category = "Subsonic")
+	ESubsonicExecutionScope Scope = ESubsonicExecutionScope::Collection;
+
+	UPROPERTY(EditAnywhere, Category = "Subsonic")
+	ESubsonicAudioComponentAccess Access = ESubsonicAudioComponentAccess::FindOrAdd;
+
+	UPROPERTY(EditAnywhere, Category = "Subsonic")
+	TArray<TInstancedStruct<FSubsonicEventAction_AudioComponentModifierBase>> Modifiers;
+};
+
+USTRUCT(MinimalAPI, BlueprintType, DisplayName = "Play Sound")
+struct FSubsonicEventAction_AudioComponentPlay : public FSubsonicEventActionBase
+{
+	GENERATED_BODY()
+
+	virtual ~FSubsonicEventAction_AudioComponentPlay() = default;
+
+	virtual void Execute(const UE::Subsonic::FActionHandle& InHandle) const override;
+
+public:
+	// Addressable name of audio component instance.
+	UPROPERTY(EditAnywhere, Category = "Subsonic")
+	FName Name = "Sound01";
+
+	// Sound to play.
+	UPROPERTY(EditAnywhere, Category = "Subsonic")
+	TObjectPtr<USoundBase> Sound;
+
+	// Scope of component accessor.  If set to event, provided name is used
+	// to address the same component within the event, collection the same component
+	// within the collection, etc.
+	UPROPERTY(EditAnywhere, Category = "Subsonic")
+	ESubsonicExecutionScope Scope = ESubsonicExecutionScope::Collection;
+
+	UPROPERTY(EditAnywhere, Category = "Subsonic")
+	ESubsonicAudioComponentAccess Access = ESubsonicAudioComponentAccess::FindOrAdd;
+};
+
+USTRUCT(MinimalAPI, BlueprintType, DisplayName = "Stop Sound")
+struct FSubsonicEventAction_AudioComponentStop : public FSubsonicEventActionBase
+{
+	GENERATED_BODY()
+
+	virtual ~FSubsonicEventAction_AudioComponentStop() = default;
+
+	virtual void Execute(const UE::Subsonic::FActionHandle& InHandle) const override;
+
+public:
+	// Addressable name of component/sound instance.
+	UPROPERTY(EditAnywhere, Category = "Subsonic")
+	FName Name = "Sound01";
+
+	// Scope of component to play.  If set to event, provided name is used
+	// to address the existing component within the event, collection the existing component
+	// within the collection, etc.
+	UPROPERTY(EditAnywhere, Category = "Subsonic")
+	ESubsonicExecutionScope Scope = ESubsonicExecutionScope::Collection;
+};
+
+USTRUCT(MinimalAPI, BlueprintType, meta = (Abstract, Hidden))
+struct FSubsonicEventAction_AudioComponentModifierBase
+{
+	GENERATED_BODY()
+
+public:
+	virtual ~FSubsonicEventAction_AudioComponentModifierBase() = default;
+
+	virtual void Execute(const UE::Subsonic::FActionHandle& InHandle, UAudioComponent& AudioComponent) const { }
+};
+
+USTRUCT(MinimalAPI, BlueprintType, DisplayName = "Play")
+struct FSubsonicEventAction_AudioComponentModifier_Play : public FSubsonicEventAction_AudioComponentModifierBase
+{
+	GENERATED_BODY()
+
+public:
+	virtual ~FSubsonicEventAction_AudioComponentModifier_Play() = default;
+
+	virtual void Execute(const UE::Subsonic::FActionHandle& InHandle, UAudioComponent& AudioComponent) const override;
+
+	// Sound to play.
+	UPROPERTY(EditAnywhere, Category = "Subsonic")
+	float StartTime = 0.0f;
+};
+
+USTRUCT(MinimalAPI, BlueprintType, DisplayName = "Set Attenuation")
+struct FSubsonicEventAction_AudioComponentModifier_SetAttenuation : public FSubsonicEventAction_AudioComponentModifierBase
+{
+	GENERATED_BODY()
+
+public:
+	virtual ~FSubsonicEventAction_AudioComponentModifier_SetAttenuation() = default;
+
+	virtual void Execute(const UE::Subsonic::FActionHandle& InHandle, UAudioComponent& AudioComponent) const override;
+
+	UPROPERTY(EditAnywhere, Category = "Subsonic")
+	TObjectPtr<USoundAttenuation> Attenuation;
+};
+
+USTRUCT(MinimalAPI, BlueprintType, DisplayName = "Set Concurrency")
+struct FSubsonicEventAction_AudioComponentModifier_SetConcurrency : public FSubsonicEventAction_AudioComponentModifierBase
+{
+	GENERATED_BODY()
+
+public:
+	virtual ~FSubsonicEventAction_AudioComponentModifier_SetConcurrency() = default;
+
+	virtual void Execute(const UE::Subsonic::FActionHandle& InHandle, UAudioComponent& AudioComponent) const override;
+
+	UPROPERTY(EditAnywhere, Category = "Subsonic")
+	TSet<TObjectPtr<USoundConcurrency>> Concurrency;
+};
+
+USTRUCT(MinimalAPI, BlueprintType, DisplayName = "Set Modulation Routing")
+struct FSubsonicEventAction_AudioComponentModifier_SetModulationRouting: public FSubsonicEventAction_AudioComponentModifierBase
+{
+	GENERATED_BODY()
+
+public:
+	virtual ~FSubsonicEventAction_AudioComponentModifier_SetModulationRouting() = default;
+
+	virtual void Execute(const UE::Subsonic::FActionHandle& InHandle, UAudioComponent& AudioComponent) const override;
+
+	UPROPERTY(EditAnywhere, Category = "Subsonic")
+	TSet<TObjectPtr<USoundModulatorBase>> Modulators;
+
+	UPROPERTY(EditAnywhere, Category = "Subsonic")
+	EModulationDestination Destination = EModulationDestination::Volume;
+	
+	UPROPERTY(EditAnywhere, Category = "Subsonic")
+	EModulationRouting RoutingMethod = EModulationRouting::Inherit;
+};
+
+USTRUCT(MinimalAPI, BlueprintType, DisplayName = "Set Sound")
+struct FSubsonicEventAction_AudioComponentModifier_SetSound : public FSubsonicEventAction_AudioComponentModifierBase
+{
+	GENERATED_BODY()
+
+public:
+	virtual ~FSubsonicEventAction_AudioComponentModifier_SetSound() = default;
+
+	virtual void Execute(const UE::Subsonic::FActionHandle& InHandle, UAudioComponent& AudioComponent) const override;
+
+	// Sound to play.
+	UPROPERTY(EditAnywhere, Category = "Subsonic")
+	TObjectPtr<USoundBase> Sound;
+};
+
+USTRUCT(MinimalAPI, BlueprintType, DisplayName = "Stop")
+struct FSubsonicEventAction_AudioComponentModifier_Stop : public FSubsonicEventAction_AudioComponentModifierBase
+{
+	GENERATED_BODY()
+
+public:
+	virtual ~FSubsonicEventAction_AudioComponentModifier_Stop() = default;
+
+	virtual void Execute(const UE::Subsonic::FActionHandle& InHandle, UAudioComponent& AudioComponent) const override;
+};
diff --git a/Engine/Plugins/Experimental/Subsonic/Source/Subsonic/Private/CoreEventSystems/SubsonicAudioComponentSubsystem.cpp b/Engine/Plugins/Experimental/Subsonic/Source/Subsonic/Private/CoreEventSystems/SubsonicAudioComponentSubsystem.cpp
new file mode 100644
--- /dev/null
+++ b/Engine/Plugins/Experimental/Subsonic/Source/Subsonic/Private/CoreEventSystems/SubsonicAudioComponentSubsystem.cpp
@@ -0,0 +1,430 @@
+// Copyright Epic Games, Inc. All Rights Reserved.
+
+#include "CoreEventSystems/SubsonicAudioComponentSubsystem.h"
+
+#include "Algo/Transform.h"
+#include "AudioDeviceManager.h"
+#include "AudioDevice.h"
+#include "Engine/Engine.h"
+#include "SubsonicLog.h"
+#include "SubsonicSubsystem.h"
+#include "UObject/UObjectGlobals.h"
+
+
+namespace UE::Subsonic
+{
+	namespace SubsystemPrivate
+	{
+		UAudioComponent* CreateComponent(USubsonicAudioComponentSubsystem& Subsystem, FName Name)
+		{
+			const FName ComponentName = MakeUniqueObjectName(&Subsystem, UAudioComponent::StaticClass(), Name);
+			UAudioComponent* NewComponent = NewObject<UAudioComponent>(&Subsystem, ComponentName);
+
+			check(NewComponent);
+			NewComponent->AudioDeviceID = Subsystem.GetAudioDeviceId();
+			NewComponent->bAutoDestroy = true;
+			NewComponent->bIsUISound = true;
+			NewComponent->bAllowSpatialization = false;
+			NewComponent->bReverb = false;
+			NewComponent->bCenterChannelOnly = false;
+
+			return NewComponent;
+		}
+
+		UAudioComponent* AddCacheComponent(USubsonicAudioComponentSubsystem& Subsystem, FInstanceAudioComponents& InstanceComponents, FName Name, bool bReleaseExisting, bool& bOutStoppedComponent)
+		{
+			bOutStoppedComponent = false;
+
+			if (bReleaseExisting)
+			{
+				if (UAudioComponent* Component = InstanceComponents.ActiveComponent.Get(); Component && Component->IsPlaying())
+				{
+					if (USoundBase* Sound = Component->GetSound(); Sound && !Sound->IsOneShot())
+					{
+						bOutStoppedComponent = true;
+						Component->Stop();
+					}
+				}
+
+				InstanceComponents.ReleaseComponent();
+			}
+
+			TStrongObjectPtr<UAudioComponent>& ComponentRef = InstanceComponents.ActiveComponent;
+			if (!ComponentRef.IsValid())
+			{
+				ComponentRef.Reset(CreateComponent(Subsystem, Name));
+			}
+
+			return ComponentRef.Get();
+		}
+
+		template <typename HandleType>
+		UAudioComponent* AddCacheComponent(
+			USubsonicAudioComponentSubsystem& Subsystem,
+			const HandleType& InHandle,
+			TMap<HandleType, FAudioComponentsCache>& OutComponentCache,
+			FName Name,
+			bool bReleaseExisting)
+		{
+			UAudioComponent* CacheComponent = nullptr;
+
+			if (FAudioDeviceManager* AudioDeviceManager = FAudioDeviceManager::Get())
+			{
+				if (FAudioDevice* AudioDevice = AudioDeviceManager->GetAudioDeviceRaw(Subsystem.GetAudioDeviceId()))
+				{
+					FAudioComponentsCache& Cache = OutComponentCache.FindOrAdd(InHandle);
+					FInstanceAudioComponents& InstanceComponents = Cache.FindOrAdd(Name);
+
+					bool bStoppedComponent = false;
+					CacheComponent = AddCacheComponent(Subsystem, InstanceComponents, Name, bReleaseExisting, bStoppedComponent);
+					if (bStoppedComponent)
+					{
+						UE_LOG(LogSubsonic, Warning,
+							TEXT("Stopping released sound with handle '%s': "
+								"Sound is not one shot and continued play would result in leaked playback."),
+							*InHandle.ToString())
+					}
+					auto RemoveComponentLambda = [&Subsystem, FinishHandle = InHandle, &OutComponentCache, Name](UAudioComponent* FinishingComponent)
+					{
+						using namespace UE::Subsonic;
+
+						if (FAudioComponentsCache* Cache = OutComponentCache.Find(FinishHandle))
+						{
+							if (FInstanceAudioComponents* InstanceComponents = Cache->Find(Name))
+							{
+								if (InstanceComponents->ActiveComponent.Get() == FinishingComponent)
+								{
+									InstanceComponents->ActiveComponent.Reset();
+								}
+
+								InstanceComponents->ReleasedComponents.RemoveAllSwap([&FinishingComponent](const TStrongObjectPtr<UAudioComponent>& Released)
+								{
+									return Released.Get() == FinishingComponent;
+								});
+
+								if (InstanceComponents->IsEmpty())
+								{
+									Cache->Remove(Name);
+								}
+							}
+							if (Cache->IsEmpty())
+							{
+								OutComponentCache.Remove(FinishHandle);
+							}
+						}
+					};
+					CacheComponent->OnAudioFinishedNative.AddLambda(RemoveComponentLambda);
+				}
+			}
+
+			return CacheComponent;
+		}
+
+		template <typename HandleType>
+		UAudioComponent* FindCacheComponent(const TMap<HandleType, FAudioComponentsCache>& ComponentCache, const HandleType& InHandle, FName Name)
+		{
+			if (const FAudioComponentsCache* Cache = ComponentCache.Find(InHandle))
+			{
+				if (const FInstanceAudioComponents* InstanceComponents = Cache->Find(Name))
+				{
+					return InstanceComponents->ActiveComponent.Get();
+				}
+			}
+
+			return nullptr;
+		}
+
+		template <typename HandleType>
+		bool ContainsCacheComponent(const TMap<HandleType, FAudioComponentsCache>& ComponentCache, const HandleType& InHandle, FName Name)
+		{
+			if (const FAudioComponentsCache* Cache = ComponentCache.Find(InHandle))
+			{
+				if (const FInstanceAudioComponents* InstanceComponents = Cache->Find(Name))
+				{
+					return InstanceComponents->ActiveComponent.IsValid();
+				}
+			}
+
+			return false;
+		}
+
+		template <typename HandleType>
+		void RemoveCacheComponent(TMap<HandleType, FAudioComponentsCache>& ComponentCache, const HandleType& InHandle, FName Name)
+		{
+			if (FAudioComponentsCache* Cache = ComponentCache.Find(InHandle))
+			{
+				if (FInstanceAudioComponents* InstanceComponents = Cache->Find(Name))
+				{
+					InstanceComponents->StopAll();
+				}
+
+				Cache->Remove(Name);
+				if (Cache->IsEmpty())
+				{
+					ComponentCache.Remove(InHandle);
+				}
+			}
+		}
+	} // namespace SubsystemPrivate
+
+	bool FInstanceAudioComponents::IsEmpty() const
+	{
+		return !ActiveComponent.IsValid() && ReleasedComponents.IsEmpty();
+	}
+
+	void FInstanceAudioComponents::ReleaseComponent()
+	{
+		if (ActiveComponent.IsValid() && ActiveComponent->IsPlaying())
+		{
+			ReleasedComponents.Add(ActiveComponent);
+		}
+		ActiveComponent.Reset();
+	}
+
+	void FInstanceAudioComponents::StopAll()
+	{
+		if (UAudioComponent* Component = ActiveComponent.Get())
+		{
+			Component->Stop();
+		}
+
+		for (TStrongObjectPtr<UAudioComponent>& ReleasedComponent : ReleasedComponents)
+		{
+			if (UAudioComponent* Component = ReleasedComponent.Get())
+			{
+				Component->Stop();
+			}
+		}
+	}
+} // UE::Subsonic
+
+void USubsonicAudioComponentSubsystem::Initialize(FSubsystemCollectionBase& Collection)
+{
+	if (GEngine)
+	{
+		if (USubsonicSubsystem* EngineSubsystem = GEngine->GetEngineSubsystem<USubsonicSubsystem>())
+		{
+			EngineSubsystem->RegisterEventSystem(this);
+		}
+	}
+}
+
+void USubsonicAudioComponentSubsystem::Deinitialize()
+{
+	using namespace UE::Subsonic;
+
+	auto StopAllCacheComponents = [](FAudioComponentsCache& Cache)
+	{
+		for (TPair<FName, FInstanceAudioComponents>& Pair : Cache)
+		{
+			FInstanceAudioComponents& InstanceComponents = Pair.Value;
+			InstanceComponents.StopAll();
+		}
+	};
+
+
+	for (TPair<FCollectionHandle, FAudioComponentsCache>& HandleCachePair : CollectionComponents)
+	{
+		StopAllCacheComponents(HandleCachePair.Value);
+	}
+
+	for (TPair<FEventHandle, FAudioComponentsCache>& HandleCachePair : EventComponents)
+	{
+		StopAllCacheComponents(HandleCachePair.Value);
+	}
+
+	for (TPair<FActionHandle, FAudioComponentsCache>& HandleCachePair : ActionComponents)
+	{
+		StopAllCacheComponents(HandleCachePair.Value);
+	}
+
+	CollectionComponents.Reset();
+	EventComponents.Reset();
+	ActionComponents.Reset();
+
+	if (GEngine)
+	{
+		if (USubsonicSubsystem* EngineSubsystem = GEngine->GetEngineSubsystem<USubsonicSubsystem>())
+		{
+			EngineSubsystem->UnregisterEventSystem(this);
+		}
+	}
+}
+
+UAudioComponent* USubsonicAudioComponentSubsystem::AddComponent(FName Name, bool bReleaseExisting)
+{
+	using namespace UE::Subsonic;
+
+	FInstanceAudioComponents& InstanceComponents = GlobalComponents.FindOrAdd(Name);
+
+	bool bStoppedComponent = false;
+	if (UAudioComponent* CacheComponent = SubsystemPrivate::AddCacheComponent(*this, InstanceComponents, Name, bReleaseExisting, bStoppedComponent))
+	{
+		if (bStoppedComponent)
+		{
+			UE_LOG(LogSubsonic, Warning,
+				TEXT("Stopping released global sound '%s': "
+					"Sound is not one shot and continued play would result in leaked playback."),
+				*Name.ToString())
+		}
+
+		auto RemoveComponentLambda = [this, Name](UAudioComponent* FinishingComponent)
+		{
+			using namespace UE::Subsonic;
+
+			if (FInstanceAudioComponents* Components = GlobalComponents.Find(Name))
+			{
+				if (Components->ActiveComponent.Get() == FinishingComponent)
+				{
+					Components->ActiveComponent.Reset();
+				}
+
+				Components->ReleasedComponents.RemoveAllSwap([&FinishingComponent](const TStrongObjectPtr<UAudioComponent>& Released)
+					{
+						return Released.Get() == FinishingComponent;
+					});
+
+				if (Components->IsEmpty())
+				{
+					GlobalComponents.Remove(Name);
+				}
+			}
+		};
+		CacheComponent->OnAudioFinishedNative.AddLambda(RemoveComponentLambda);
+		return CacheComponent;
+	}
+
+	return nullptr;
+}
+
+UAudioComponent* USubsonicAudioComponentSubsystem::AddComponent(const UE::Subsonic::FCollectionHandle& InHandle, FName Name, bool bReleaseExisting)
+{
+	using namespace UE::Subsonic;
+	return SubsystemPrivate::AddCacheComponent(*this, InHandle, CollectionComponents, Name, bReleaseExisting);
+}
+
+UAudioComponent* USubsonicAudioComponentSubsystem::AddComponent(const UE::Subsonic::FEventHandle& InHandle, FName Name, bool bReleaseExisting)
+{
+	using namespace UE::Subsonic;
+	return SubsystemPrivate::AddCacheComponent(*this, InHandle, EventComponents, Name, bReleaseExisting);
+}
+
+UAudioComponent* USubsonicAudioComponentSubsystem::AddComponent(const UE::Subsonic::FActionHandle& InHandle, FName Name, bool bReleaseExisting)
+{
+	using namespace UE::Subsonic;
+	return SubsystemPrivate::AddCacheComponent(*this, InHandle, ActionComponents, Name, bReleaseExisting);
+}
+
+bool USubsonicAudioComponentSubsystem::ContainsComponent(FName Name) const
+{
+	return GlobalComponents.Contains(Name);
+}
+
+bool USubsonicAudioComponentSubsystem::ContainsComponent(const UE::Subsonic::FCollectionHandle& InHandle, FName Name) const
+{
+	using namespace UE::Subsonic;
+	return SubsystemPrivate::ContainsCacheComponent(CollectionComponents, InHandle, Name);
+}
+
+bool USubsonicAudioComponentSubsystem::ContainsComponent(const UE::Subsonic::FEventHandle& InHandle, FName Name) const
+{
+	using namespace UE::Subsonic;
+	return SubsystemPrivate::ContainsCacheComponent(EventComponents, InHandle, Name);
+}
+
+bool USubsonicAudioComponentSubsystem::ContainsComponent(const UE::Subsonic::FActionHandle& InHandle, FName Name) const
+{
+	using namespace UE::Subsonic;
+	return SubsystemPrivate::ContainsCacheComponent(ActionComponents, InHandle, Name);
+}
+
+UAudioComponent* USubsonicAudioComponentSubsystem::FindComponent(FName Name)
+{
+	using namespace UE::Subsonic;
+	if (FInstanceAudioComponents* InstanceComponents = GlobalComponents.Find(Name))
+	{
+		return InstanceComponents->ActiveComponent.Get();
+	}
+
+	return nullptr;
+}
+
+UAudioComponent* USubsonicAudioComponentSubsystem::FindComponent(const UE::Subsonic::FCollectionHandle& InHandle, FName Name)
+{
+	using namespace UE::Subsonic;
+	return SubsystemPrivate::FindCacheComponent(CollectionComponents, InHandle, Name);
+}
+
+UAudioComponent* USubsonicAudioComponentSubsystem::FindComponent(const UE::Subsonic::FEventHandle& InHandle, FName Name)
+{
+	using namespace UE::Subsonic;
+	return SubsystemPrivate::FindCacheComponent(EventComponents, InHandle, Name);
+}
+
+UAudioComponent* USubsonicAudioComponentSubsystem::FindComponent(const UE::Subsonic::FActionHandle& InHandle, FName Name)
+{
+	using namespace UE::Subsonic;
+	return SubsystemPrivate::FindCacheComponent(ActionComponents, InHandle, Name);
+}
+
+void USubsonicAudioComponentSubsystem::RemoveComponent(FName Name)
+{
+	using namespace UE::Subsonic;
+
+	if (FInstanceAudioComponents* InstanceComponents = GlobalComponents.Find(Name))
+	{
+		InstanceComponents->StopAll();
+		GlobalComponents.Remove(Name);
+	}
+}
+
+void USubsonicAudioComponentSubsystem::RemoveComponent(const UE::Subsonic::FCollectionHandle& InHandle, FName Name)
+{
+	using namespace UE::Subsonic;
+	return SubsystemPrivate::RemoveCacheComponent(CollectionComponents, InHandle, Name);
+}
+
+void USubsonicAudioComponentSubsystem::RemoveComponent(const UE::Subsonic::FEventHandle& InHandle, FName Name)
+{
+	using namespace UE::Subsonic;
+	return SubsystemPrivate::RemoveCacheComponent(EventComponents, InHandle, Name);
+}
+
+void USubsonicAudioComponentSubsystem::OnCollectionRegistered(const UE::Subsonic::FCollectionHandle& InHandle)
+{
+}
+
+void USubsonicAudioComponentSubsystem::OnCollectionUnregistered(const UE::Subsonic::FCollectionHandle& InHandle)
+{
+	using namespace UE::Subsonic;
+
+	if (FAudioComponentsCache* Cache = CollectionComponents.Find(InHandle))
+	{
+		for (TPair<FName, FInstanceAudioComponents>& Pair : *Cache)
+		{
+			FInstanceAudioComponents& InstanceComponents = Pair.Value;
+			if (UAudioComponent* Component = InstanceComponents.ActiveComponent.Get())
+			{
+				Component->Stop();
+			}
+
+			for (TStrongObjectPtr<UAudioComponent>& ReleasedComponent: InstanceComponents.ReleasedComponents)
+			{
+				if (UAudioComponent* Component = ReleasedComponent.Get())
+				{
+					Component->Stop();
+				}
+			}
+		}
+
+		CollectionComponents.Remove(InHandle);
+	}
+}
+
+void USubsonicAudioComponentSubsystem::OnEventPreExecute(const UE::Subsonic::FEventHandle& InHandle)
+{
+}
+
+void USubsonicAudioComponentSubsystem::OnEventPostExecute(const UE::Subsonic::FEventHandle& InHandle)
+{
+}
diff --git a/Engine/Plugins/Experimental/Subsonic/Source/Subsonic/Private/CoreEventSystems/SubsonicAudioComponentSubsystem.h b/Engine/Plugins/Experimental/Subsonic/Source/Subsonic/Private/CoreEventSystems/SubsonicAudioComponentSubsystem.h
new file mode 100644
--- /dev/null
+++ b/Engine/Plugins/Experimental/Subsonic/Source/Subsonic/Private/CoreEventSystems/SubsonicAudioComponentSubsystem.h
@@ -0,0 +1,91 @@
+// Copyright Epic Games, Inc. All Rights Reserved.
+#pragma once
+
+#include "Subsystems/EngineSubsystem.h"
+
+#include "Components/AudioComponent.h"
+#include "Containers/Map.h"
+#include "SubsonicEventSystemInterface.h"
+#include "SubsonicHandles.h"
+#include "Subsystems/AudioEngineSubsystem.h"
+#include "UObject/GCObject.h"
+
+#include "SubsonicAudioComponentSubsystem.generated.h"
+
+#define UE_API SUBSONIC_API
+
+// Forward Declarations
+class FSubsonicAudioComponentSubsystemImpl;
+
+namespace UE::Subsonic
+{
+	struct FInstanceAudioComponents
+	{
+		// Currently addressable component via handle
+		TStrongObjectPtr<UAudioComponent> ActiveComponent;
+
+		// Components no longer addressable from actions but actively
+		// playing. Continually tracked here to avoid GC while finishing.
+		TArray<TStrongObjectPtr<UAudioComponent>> ReleasedComponents;
+
+		bool IsEmpty() const;
+		void ReleaseComponent();
+		void StopAll();
+	};
+
+	using FAudioComponentsCache = TMap<FName, FInstanceAudioComponents>;
+} // namespace UE::Subsonic
+
+/**
+ * Interface to AudioComponent playback Subsonic implementation.
+ */
+UCLASS(MinimalAPI, BlueprintType)
+class USubsonicAudioComponentSubsystem final : public UAudioEngineSubsystem, public ISubsonicEventSystemInterface
+{
+	GENERATED_BODY()
+
+private:
+	using FCollectionHandle = UE::Subsonic::FCollectionHandle;
+	using FEventHandle = UE::Subsonic::FEventHandle;
+	using FActionHandle = UE::Subsonic::FActionHandle;
+	using FAudioComponentsCache = UE::Subsonic::FAudioComponentsCache;
+
+public:
+	//~ Begin USubsystem interface
+	UE_API virtual void Initialize(FSubsystemCollectionBase& Collection) override;
+	UE_API virtual void Deinitialize() override;
+	//~ End USubsystem interface
+
+	UE_API virtual void OnEventPreExecute(const FEventHandle& InHandle) override;
+	UE_API virtual void OnEventPostExecute(const FEventHandle& InHandle) override;
+
+	UE_API virtual void OnCollectionRegistered(const FCollectionHandle& InEvent) override;
+	UE_API virtual void OnCollectionUnregistered(const FCollectionHandle& InEvent) override;
+
+	UE_INTERNAL UAudioComponent* AddComponent(FName Name, bool bReleaseExisting = true);
+	UE_INTERNAL UAudioComponent* AddComponent(const FCollectionHandle& InHandle, FName Name, bool bReleaseExisting = true);
+	UE_INTERNAL UAudioComponent* AddComponent(const FEventHandle& InHandle, FName Name, bool bReleaseExisting = true);
+	UE_INTERNAL UAudioComponent* AddComponent(const FActionHandle& InHandle, FName Name, bool bReleaseExisting = true);
+
+	UE_INTERNAL bool ContainsComponent(FName Name) const;
+	UE_INTERNAL bool ContainsComponent(const FCollectionHandle& InHandle, FName Name) const;
+	UE_INTERNAL bool ContainsComponent(const FEventHandle& InHandle, FName Name) const;
+	UE_INTERNAL bool ContainsComponent(const FActionHandle& InHandle, FName Name) const;
+
+	UE_INTERNAL UAudioComponent* FindComponent(FName Name);
+	UE_INTERNAL UAudioComponent* FindComponent(const FCollectionHandle& InHandle, FName Name);
+	UE_INTERNAL UAudioComponent* FindComponent(const FEventHandle& InHandle, FName Name);
+	UE_INTERNAL UAudioComponent* FindComponent(const FActionHandle& InHandle, FName Name);
+
+	UE_INTERNAL void RemoveComponent(FName Name);
+	UE_INTERNAL void RemoveComponent(const FCollectionHandle& InHandle, FName Name);
+	UE_INTERNAL void RemoveComponent(const FEventHandle& InHandle, FName Name);
+
+private:
+	FAudioComponentsCache GlobalComponents;
+	TMap<FCollectionHandle, FAudioComponentsCache> CollectionComponents;
+	TMap<FEventHandle, FAudioComponentsCache> EventComponents;
+	TMap<FActionHandle, FAudioComponentsCache> ActionComponents;
+};
+
+#undef UE_API
diff --git a/Engine/Plugins/Experimental/Subsonic/Source/Subsonic/Private/SubsonicEventCollection.cpp b/Engine/Plugins/Experimental/Subsonic/Source/Subsonic/Private/SubsonicEventCollection.cpp
new file mode 100644
--- /dev/null
+++ b/Engine/Plugins/Experimental/Subsonic/Source/Subsonic/Private/SubsonicEventCollection.cpp
@@ -0,0 +1,161 @@
+// Copyright Epic Games, Inc. All Rights Reserved.
+
+#include "SubsonicEventCollection.h"
+
+#include "Engine/Engine.h"
+#include "NativeGameplayTags.h"
+#include "StructUtils/StructView.h"
+#include "SubsonicHandles.h"
+#include "SubsonicLog.h"
+#include "SubsonicSubsystem.h"
+
+
+namespace UE::Subsonic
+{
+	UE_DEFINE_GAMEPLAY_TAG(TAG_SubsonicEventCorePlay, "Subsonic.Event.Play");
+	UE_DEFINE_GAMEPLAY_TAG(TAG_SubsonicEventCoreStop, "Subsonic.Event.Stop");
+} // namespace UE::Subsonic
+
+
+void FSubsonicEvent::Execute(const UE::Subsonic::FEventHandle& InHandle) const
+{
+	using namespace UE::Subsonic;
+
+	FActionHandle ActionHandle { .Event = InHandle };
+	for (int32 Index = 0; Index < Actions.Num(); ++Index)
+	{
+		ActionHandle.Index = Index;
+
+		TConstStructView<FSubsonicEventActionBase> Action = Actions[Index];
+		if (const FSubsonicEventActionBase* ActionPtr = Action.GetPtr())
+		{
+			ActionPtr->Execute(ActionHandle);
+		}
+		else
+		{
+			UE_LOG(LogSubsonic, Warning, TEXT("Action not specified: Skipping action at '%s'"), *ActionHandle.ToString());
+		}
+	}
+}
+
+UE::Subsonic::FCollectionHandle USubsonicEventCollection::GetHandle(const FSubsonicEventCollectionExecutor& Executor) const
+{
+	using namespace UE::Subsonic;
+
+	FName CollectionName;
+
+#if WITH_EDITORONLY_DATA
+	if (const UPackage* CollectionPackage = GetPackage())
+	{
+		const FTopLevelAssetPath AssetPath(CollectionPackage);
+		CollectionName = AssetPath.GetPackageName();
+	}
+#endif // WITH_EDITORONLY_DATA
+
+	return FCollectionHandle
+	{
+#if WITH_EDITORONLY_DATA
+		.CollectionName = CollectionName,
+#endif // WITH_EDITORONLY_DATA
+		.CollectionId = GetUniqueID(),
+		.DeviceId = Executor.GetDeviceId(),
+		.ExecutorId = Executor.GetExecutorId()
+	};
+}
+
+void USubsonicEventCollection::PostLoad()
+{
+	Super::PostLoad();
+
+	using namespace UE::Subsonic;
+
+	if (GEngine)
+	{
+		if (USubsonicSubsystem* Subsystem = GEngine->GetEngineSubsystem<USubsonicSubsystem>())
+		{
+			FName CollectionName;
+
+#if WITH_EDITORONLY_DATA
+			if (const UPackage* CollectionPackage = GetPackage())
+			{
+				const FTopLevelAssetPath AssetPath(CollectionPackage);
+				CollectionName = AssetPath.GetPackageName();
+			}
+#endif // WITH_EDITORONLY_DATA
+
+			Subsystem->RegisterCollection(FCollectionHandle
+			{
+#if WITH_EDITORONLY_DATA
+				.CollectionName = CollectionName,
+#endif // WITH_EDITORONLY_DATA
+				.CollectionId = GetUniqueID(),
+			});
+		}
+	}
+}
+
+void USubsonicEventCollection::BeginDestroy()
+{
+	using namespace UE::Subsonic;
+
+	if (GEngine)
+	{
+		if (USubsonicSubsystem* Subsystem = GEngine->GetEngineSubsystem<USubsonicSubsystem>())
+		{
+			FName CollectionName;
+
+#if WITH_EDITORONLY_DATA
+			if (const UPackage* CollectionPackage = GetPackage())
+			{
+				const FTopLevelAssetPath AssetPath(CollectionPackage);
+				CollectionName = AssetPath.GetPackageName();
+			}
+#endif // WITH_EDITORONLY_DATA
+
+			Subsystem->UnregisterCollection(FCollectionHandle
+			{
+#if WITH_EDITORONLY_DATA
+				.CollectionName = CollectionName,
+#endif // WITH_EDITORONLY_DATA
+				.CollectionId = GetUniqueID(),
+			});
+		}
+	}
+
+	Super::BeginDestroy();
+}
+
+uint32 FSubsonicEventCollectionExecutor::MaxExecutorId = INDEX_NONE;
+
+FSubsonicEventCollectionExecutor::FSubsonicEventCollectionExecutor(const USubsonicEventCollection* InCollection, Audio::FDeviceId InDeviceId)
+	: Collection(InCollection)
+	, DeviceId(InDeviceId)
+	, ExecutorId(++MaxExecutorId)
+{
+}
+
+Audio::FDeviceId FSubsonicEventCollectionExecutor::GetDeviceId() const
+{
+	return DeviceId;
+}
+
+const USubsonicEventCollection& FSubsonicEventCollectionExecutor::GetCollectionChecked() const
+{
+	check(Collection.Get());
+	return *Collection.Get();
+}
+
+const USubsonicEventCollection* FSubsonicEventCollectionExecutor::GetCollection() const
+{
+	return Collection.Get();
+}
+
+uint32 FSubsonicEventCollectionExecutor::GetExecutorId() const
+{
+	return ExecutorId;
+}
+
+bool FSubsonicEventCollectionExecutor::IsValid() const
+{
+	return ExecutorId != INDEX_NONE && Collection.Get() != nullptr;
+}
diff --git a/Engine/Plugins/Experimental/Subsonic/Source/Subsonic/Private/SubsonicEventHandler.cpp b/Engine/Plugins/Experimental/Subsonic/Source/Subsonic/Private/SubsonicEventHandler.cpp
deleted file mode 100644
diff --git a/Engine/Plugins/Experimental/Subsonic/Source/Subsonic/Private/SubsonicEventSystemInterface.h b/Engine/Plugins/Experimental/Subsonic/Source/Subsonic/Private/SubsonicEventSystemInterface.h
new file mode 100644
--- /dev/null
+++ b/Engine/Plugins/Experimental/Subsonic/Source/Subsonic/Private/SubsonicEventSystemInterface.h
@@ -0,0 +1,26 @@
+// Copyright Epic Games, Inc. All Rights Reserved.
+#pragma once
+
+#include "UObject/Interface.h"
+#include "SubsonicHandles.h"
+
+#include "SubsonicEventSystemInterface.generated.h"
+
+// Interface for all subscribers that respond to Subsonic events.
+UINTERFACE(MinimalAPI, BlueprintType, meta = (DisplayName = "Subsonic Event System Interface", CannotImplementInterfaceInBlueprint))
+class USubsonicEventSystemInterface : public UInterface
+{
+	GENERATED_BODY()
+};
+
+class ISubsonicEventSystemInterface : public IInterface
+{
+	GENERATED_BODY()
+
+public:
+	virtual void OnEventPreExecute(const UE::Subsonic::FEventHandle& InHandle) = 0;
+	virtual void OnEventPostExecute(const UE::Subsonic::FEventHandle& InHandle) = 0;
+
+	virtual void OnCollectionRegistered(const UE::Subsonic::FCollectionHandle& InEvent) = 0;
+	virtual void OnCollectionUnregistered(const UE::Subsonic::FCollectionHandle& InEvent) = 0;
+};
diff --git a/Engine/Plugins/Experimental/Subsonic/Source/Subsonic/Private/SubsonicLog.h b/Engine/Plugins/Experimental/Subsonic/Source/Subsonic/Private/SubsonicLog.h
new file mode 100644
--- /dev/null
+++ b/Engine/Plugins/Experimental/Subsonic/Source/Subsonic/Private/SubsonicLog.h
@@ -0,0 +1,6 @@
+// Copyright Epic Games, Inc. All Rights Reserved.
+#pragma once
+
+#include "Logging/LogCategory.h"
+
+DECLARE_LOG_CATEGORY_EXTERN(LogSubsonic, Log, All);
diff --git a/Engine/Plugins/Experimental/Subsonic/Source/Subsonic/Private/SubsonicModule.cpp b/Engine/Plugins/Experimental/Subsonic/Source/Subsonic/Private/SubsonicModule.cpp
--- a/Engine/Plugins/Experimental/Subsonic/Source/Subsonic/Private/SubsonicModule.cpp
+++ b/Engine/Plugins/Experimental/Subsonic/Source/Subsonic/Private/SubsonicModule.cpp
@@ -1,5 +1,8 @@
 // Copyright Epic Games, Inc. All Rights Reserved.
 
 #include "Modules/ModuleManager.h"
+#include "SubsonicLog.h"
+
+DEFINE_LOG_CATEGORY(LogSubsonic);
 
 IMPLEMENT_MODULE(FDefaultModuleImpl, Subsonic)
diff --git a/Engine/Plugins/Experimental/Subsonic/Source/Subsonic/Private/SubsonicSubsystem.cpp b/Engine/Plugins/Experimental/Subsonic/Source/Subsonic/Private/SubsonicSubsystem.cpp
--- a/Engine/Plugins/Experimental/Subsonic/Source/Subsonic/Private/SubsonicSubsystem.cpp
+++ b/Engine/Plugins/Experimental/Subsonic/Source/Subsonic/Private/SubsonicSubsystem.cpp
@@ -2,11 +2,37 @@
 
 #include "SubsonicSubsystem.h"
 
+#include "AudioDefines.h"
 #include "Engine/GameInstance.h"
+#include "Engine/Engine.h"
+#include "GameplayTagsManager.h"
+#include "Misc/CoreDefines.h"
 #include "SubsonicEventInstanceImpl.h"
 #include "SubsonicTempDataAsset.h"
+#include "SubsonicEventCollection.h"
+#include "SubsonicHandles.h"
 #include "UObject/GCObject.h"
 
+
+namespace UE::Subsonic
+{
+	namespace SubsystemPrivate
+	{
+		Audio::FDeviceId GetDeviceIdFrom(const UObject* WorldContextObject)
+		{
+			if (GEngine)
+			{
+				if (UWorld* World = GEngine->GetWorldFromContextObject(WorldContextObject, EGetWorldErrorMode::LogAndReturnNull))
+				{
+					return World->GetAudioDevice().GetDeviceID();
+				}
+			}
+
+			return INDEX_NONE;
+		};
+	} // namespace SubsystemPrivate
+} // namespace UE::Subsonic
+
 class FSubsonicSubsystemImpl : FGCObject
 {
 public:
@@ -78,6 +104,59 @@ void USubsonicSubsystem::Deinitialize()
 	Super::Deinitialize();
 }
 
+FSubsonicEventCollectionExecutor USubsonicSubsystem::CreateExecutor(const UObject* WorldContextObject, const USubsonicEventCollection* Collection)
+{
+	using namespace UE::Subsonic;
+
+	const Audio::FDeviceId DeviceId = SubsystemPrivate::GetDeviceIdFrom(WorldContextObject);
+	return FSubsonicEventCollectionExecutor(Collection, DeviceId);
+}
+
+void USubsonicSubsystem::ExecuteEvent(const FSubsonicEventCollectionExecutor& Executor, FGameplayTag EventTag)
+{
+	using namespace UE::Subsonic;
+
+	if (!Executor.IsValid())
+	{
+		return;
+	}
+
+	const USubsonicEventCollection& Collection = Executor.GetCollectionChecked();
+
+	if (const FSubsonicEvent* Event = Collection.CollectionData.Events.Find(EventTag))
+	{
+		const UE::Subsonic::FEventHandle Handle { .Collection = Collection.GetHandle(Executor), .EventName = EventTag.GetTagName() };
+
+		// Notify associated device event systems event about to begin
+		TArray<TScriptInterface<ISubsonicEventSystemInterface>> NotifiedSystems;
+		for (const TScriptInterface<ISubsonicEventSystemInterface>& EventSystem : EventSystems)
+		{
+			const UAudioEngineSubsystem* Subsystem = Cast<UAudioEngineSubsystem>(EventSystem.GetObject());
+			const Audio::FDeviceId SystemDeviceId = Subsystem
+				? Subsystem->GetAudioDeviceId()
+				: SubsystemPrivate::GetDeviceIdFrom(EventSystem.GetObject());
+
+			if (SystemDeviceId == Handle.Collection.DeviceId)
+			{
+				if (EventSystem.GetObject())
+				{
+					NotifiedSystems.Add(EventSystem);
+					EventSystem->OnEventPreExecute(Handle);
+				}
+			}
+		}
+
+		Event->Execute(Handle);
+
+		// Notify associated device event systems when event is complete
+		for (TScriptInterface<ISubsonicEventSystemInterface>& System : NotifiedSystems)
+		{
+			check(System.GetObject());
+			System->OnEventPostExecute(Handle);
+		}
+	}
+}
+
 FSubsonicEventInstance USubsonicSubsystem::PlayEvent(const UObject* WorldContextObject, const FName EventName)
 {
 	return Impl->PlayEvent(WorldContextObject, EventName);
@@ -92,4 +171,44 @@ void USubsonicSubsystem::StopEvent(FSubsonicEventInstance EventInstance)
 void USubsonicSubsystem::RegisterHandlers(USubsonicTempDataAsset* Handlers)
 {
 	Impl->RegisterHandlers(Handlers);
-}
\ No newline at end of file
+}
+
+void USubsonicSubsystem::RegisterCollection(const UE::Subsonic::FCollectionHandle& InHandle)
+{
+	using namespace UE::Subsonic;
+
+	FCollectionHandle CollectionHandle = InHandle;
+	for (TScriptInterface<ISubsonicEventSystemInterface>& EventSystem : EventSystems)
+	{
+		if (const UObject* EventSystemObj = EventSystem.GetObject())
+		{
+			CollectionHandle.DeviceId = SubsystemPrivate::GetDeviceIdFrom(EventSystemObj);
+			EventSystem->OnCollectionRegistered(CollectionHandle);
+		}
+	}
+}
+
+void USubsonicSubsystem::RegisterEventSystem(TScriptInterface<ISubsonicEventSystemInterface> EventSystem)
+{
+	EventSystems.AddUnique(EventSystem);
+}
+
+void USubsonicSubsystem::UnregisterCollection(const UE::Subsonic::FCollectionHandle& InHandle)
+{
+	using namespace UE::Subsonic;
+
+	FCollectionHandle CollectionHandle = InHandle;
+	for (TScriptInterface<ISubsonicEventSystemInterface>& EventSystem : EventSystems)
+	{
+		if (const UObject* EventSystemObj = EventSystem.GetObject())
+		{
+			CollectionHandle.DeviceId = SubsystemPrivate::GetDeviceIdFrom(EventSystemObj);
+			EventSystem->OnCollectionUnregistered(CollectionHandle);
+		}
+	}
+}
+
+void USubsonicSubsystem::UnregisterEventSystem(TScriptInterface<ISubsonicEventSystemInterface> EventSystem)
+{
+	EventSystems.Remove(EventSystem);
+}
diff --git a/Engine/Plugins/Experimental/Subsonic/Source/Subsonic/Public/SubsonicEventCollection.h b/Engine/Plugins/Experimental/Subsonic/Source/Subsonic/Public/SubsonicEventCollection.h
new file mode 100644
--- /dev/null
+++ b/Engine/Plugins/Experimental/Subsonic/Source/Subsonic/Public/SubsonicEventCollection.h
@@ -0,0 +1,99 @@
+// Copyright Epic Games, Inc. All Rights Reserved.
+#pragma once
+
+#include "AudioDefines.h"
+#include "GameplayTagContainer.h"
+#include "NativeGameplayTags.h"
+#include "StructUtils/InstancedStruct.h"
+#include "SubsonicHandles.h"
+#include "UObject/NameTypes.h"
+
+#include "SubsonicEventCollection.generated.h"
+
+#define UE_API SUBSONIC_API
+
+namespace UE::Subsonic
+{
+	// Forward Declarations
+	struct FActionHandle;
+
+	UE_API UE_DECLARE_GAMEPLAY_TAG_EXTERN(TAG_SubsonicEventCorePlay)
+	UE_API UE_DECLARE_GAMEPLAY_TAG_EXTERN(TAG_SubsonicEventCoreStop)
+} // namespace UE::Subsonic
+
+
+USTRUCT(MinimalAPI, BlueprintType, meta = (Abstract, Hidden))
+struct FSubsonicEventActionBase
+{
+	GENERATED_BODY()
+
+	virtual ~FSubsonicEventActionBase() = default;
+
+	virtual void Execute(const UE::Subsonic::FActionHandle& InHandle) const { }
+};
+
+USTRUCT(MinimalAPI, BlueprintType)
+struct FSubsonicEvent
+{
+	GENERATED_BODY()
+
+	virtual ~FSubsonicEvent() = default;
+
+	void Execute(const UE::Subsonic::FEventHandle& InHandle) const;
+
+	UPROPERTY(EditAnywhere, Category = "Subsonic")
+	TArray<TInstancedStruct<FSubsonicEventActionBase>> Actions;
+};
+
+USTRUCT(MinimalAPI, BlueprintType)
+struct FSubsonicEventCollectionDefinition
+{
+	GENERATED_BODY()
+
+public:
+	UPROPERTY(EditAnywhere, Category = "Subsonic")
+	TMap<FGameplayTag, FSubsonicEvent> Events;
+};
+
+
+UCLASS(MinimalAPI, BlueprintType)
+class USubsonicEventCollection final : public UObject
+{
+	GENERATED_BODY()
+
+public:
+	UPROPERTY(EditAnywhere, Category = "Subsonic", meta = (ShowOnlyInnerProperties))
+	FSubsonicEventCollectionDefinition CollectionData;
+
+	UE::Subsonic::FCollectionHandle GetHandle(const FSubsonicEventCollectionExecutor& Executor) const;
+
+	virtual void PostLoad() override;
+	virtual void BeginDestroy() override;
+};
+
+USTRUCT(MinimalAPI, BlueprintType)
+struct FSubsonicEventCollectionExecutor final
+{
+	GENERATED_BODY()
+
+private:
+	static uint32 MaxExecutorId;
+
+public:
+	FSubsonicEventCollectionExecutor() = default;
+	FSubsonicEventCollectionExecutor(const USubsonicEventCollection* InCollection, Audio::FDeviceId InDeviceId);
+
+	Audio::FDeviceId GetDeviceId() const;
+	const USubsonicEventCollection* GetCollection() const;
+	const USubsonicEventCollection& GetCollectionChecked() const;
+	uint32 GetExecutorId() const;
+	bool IsValid() const;
+
+private:
+	UPROPERTY(BlueprintReadOnly, Category = "Subsonic", meta = (AllowPrivateAccess))
+	TObjectPtr<const USubsonicEventCollection> Collection;
+
+	Audio::FDeviceId DeviceId = INDEX_NONE;
+	uint32 ExecutorId = INDEX_NONE;
+};
+#undef UE_API // SUBSONIC_API
diff --git a/Engine/Plugins/Experimental/Subsonic/Source/Subsonic/Public/SubsonicEventHandler.h b/Engine/Plugins/Experimental/Subsonic/Source/Subsonic/Public/SubsonicEventHandler.h
deleted file mode 100644
diff --git a/Engine/Plugins/Experimental/Subsonic/Source/Subsonic/Public/SubsonicHandles.h b/Engine/Plugins/Experimental/Subsonic/Source/Subsonic/Public/SubsonicHandles.h
new file mode 100644
--- /dev/null
+++ b/Engine/Plugins/Experimental/Subsonic/Source/Subsonic/Public/SubsonicHandles.h
@@ -0,0 +1,139 @@
+// Copyright Epic Games, Inc. All Rights Reserved.
+#pragma once
+
+#include "AudioDeviceManager.h"
+#include "AudioDevice.h"
+#include "HAL/Platform.h"
+#include "Misc/CoreDefines.h"
+#include "Templates/TypeHash.h"
+#include "UObject/NameTypes.h"
+#include "UObject/WeakObjectPtrTemplates.h"
+
+#include "SubsonicHandles.generated.h"
+
+#define UE_API SUBSONIC_API
+
+
+UENUM(BlueprintType)
+enum class ESubsonicExecutionScope : uint8
+{
+	Global,
+	Collection,
+	Event,
+	Action
+};
+
+
+namespace UE::Subsonic
+{
+	struct FCollectionHandle
+	{
+#if WITH_EDITORONLY_DATA
+		// Owning Collection definition's name (for logging/debugging)
+		FName CollectionName;
+#endif // WITH_EDITORONLY_DATA
+
+		// Owning Collection's instance id
+		uint32 CollectionId = static_cast<uint32>(INDEX_NONE);
+
+		// DeviceId corresponding to which audio device to execute event on.
+		Audio::FDeviceId DeviceId = static_cast<Audio::FDeviceId>(INDEX_NONE);
+
+		// Id of calling context executor. May be unset if handle context does not require it.
+		uint32 ExecutorId = static_cast<uint32>(INDEX_NONE);
+
+		template <typename SubsystemType>
+		SubsystemType* FindDeviceSubsystem() const
+		{
+			if (FAudioDeviceManager* DeviceManager = FAudioDeviceManager::Get())
+			{
+				FAudioDevice* AudioDevice = nullptr;
+				if (DeviceId == INDEX_NONE)
+				{
+					AudioDevice = DeviceManager->GetMainAudioDeviceRaw();
+				}
+				else
+				{
+					AudioDevice = DeviceManager->GetAudioDeviceRaw(DeviceId);
+				}
+
+				if (AudioDevice)
+				{
+					return AudioDevice->GetSubsystem<SubsystemType>();
+				}
+			}
+
+			return nullptr;
+		}
+
+		friend uint32 GetTypeHash(const FCollectionHandle& Handle)
+		{
+			const uint32 Hash1 = HashCombineFast(GetTypeHash(Handle.CollectionId), GetTypeHash(Handle.ExecutorId));
+			return HashCombineFast(Hash1, GetTypeHash(Handle.DeviceId));
+		}
+
+		friend bool operator==(const FCollectionHandle& A, const FCollectionHandle& B)
+		{
+			return A.ExecutorId == B.ExecutorId
+				&& A.CollectionId == B.CollectionId
+				&& A.DeviceId == B.DeviceId;
+		}
+
+		FString ToString() const
+		{
+#if WITH_EDITORONLY_DATA
+			return FString::Printf(TEXT("%s(%u, Exec: %u)"), *CollectionName.ToString(), CollectionId, ExecutorId);
+#else // !WITH_EDITORONLY_DATA
+			return FString::Printf(TEXT("%u, Exec: %u"), CollectionId, ExecutorId);
+#endif // !WITH_EDITORONLY_DATA
+		}
+	};
+
+	struct FEventHandle
+	{
+		FCollectionHandle Collection;
+
+		// Owning event name
+		FName EventName;
+
+		friend uint32 GetTypeHash(const FEventHandle& Handle)
+		{
+			return HashCombineFast(GetTypeHash(Handle.Collection), GetTypeHash(Handle.EventName));
+		}
+
+		friend bool operator==(const FEventHandle& A, const FEventHandle& B)
+		{
+			return A.Collection == B.Collection && A.EventName == B.EventName;
+		}
+
+		FString ToString() const
+		{
+			return FString::Printf(TEXT("%s/%s"), *Collection.ToString(), *EventName.ToString());
+		}
+	};
+
+	struct FActionHandle
+	{
+		FEventHandle Event;
+
+		// Owning action's index within a given event
+		int32 Index = INDEX_NONE;
+
+		friend uint32 GetTypeHash(const FActionHandle& Handle)
+		{
+			return HashCombineFast(GetTypeHash(Handle.Event), GetTypeHash(Handle.Index));
+		}
+
+		friend bool operator==(const FActionHandle& A, const FActionHandle& B)
+		{
+			return A.Event == B.Event && A.Index == B.Index;
+		}
+
+		FString ToString() const
+		{
+			return FString::Printf(TEXT("%s[%i]"), *Event.ToString(), Index);
+		}
+	};
+} // namespace UE::Subsonic
+
+#undef UE_API
diff --git a/Engine/Plugins/Experimental/Subsonic/Source/Subsonic/Public/SubsonicSubsystem.h b/Engine/Plugins/Experimental/Subsonic/Source/Subsonic/Public/SubsonicSubsystem.h
--- a/Engine/Plugins/Experimental/Subsonic/Source/Subsonic/Public/SubsonicSubsystem.h
+++ b/Engine/Plugins/Experimental/Subsonic/Source/Subsonic/Public/SubsonicSubsystem.h
@@ -2,6 +2,8 @@
 
 #pragma once
 
+#include "GameplayTagContainer.h"
+#include "SubsonicEventSystemInterface.h"
 #include "SubsonicTypes.h"
 #include "Subsystems/EngineSubsystem.h"
 
@@ -12,6 +14,7 @@
 class FSubsonicEventInstanceImpl;
 class FSubsonicSubsystemImpl;
 class USubsonicTempDataAsset;
+class USubsonicEventCollection;
 
 /**
  * Interface to the Subsonic plugin.
@@ -35,8 +38,24 @@ public:
 
 	UFUNCTION(BlueprintCallable, Category = "Subsonic")
 	UE_API void StopEvent(FSubsonicEventInstance EventInstance);
+
+	UE_API void RegisterEventSystem(TScriptInterface<ISubsonicEventSystemInterface> Client);
+	UE_API void RegisterCollection(const UE::Subsonic::FCollectionHandle& InHandle);
+
+	UE_API void UnregisterEventSystem(TScriptInterface<ISubsonicEventSystemInterface> Client);
+	UE_API void UnregisterCollection(const UE::Subsonic::FCollectionHandle& InHandle);
+
+	UFUNCTION(BlueprintCallable, Category = "Subsonic", meta = (WorldContext = "WorldContextObject"))
+	UE_API UPARAM(DisplayName = "Executor") FSubsonicEventCollectionExecutor CreateExecutor(const UObject* WorldContextObject, const USubsonicEventCollection* Collection);
+
+	UFUNCTION(BlueprintCallable, Category = "Subsonic")
+	UE_API void ExecuteEvent(const FSubsonicEventCollectionExecutor& Executor, FGameplayTag Event);
+
 private:
 	TPimplPtr<FSubsonicSubsystemImpl> Impl;
+
+	UPROPERTY()
+	TArray<TScriptInterface<ISubsonicEventSystemInterface>> EventSystems;
 };
 
 #undef UE_API
diff --git a/Engine/Plugins/Experimental/Subsonic/Source/Subsonic/Subsonic.Build.cs b/Engine/Plugins/Experimental/Subsonic/Source/Subsonic/Subsonic.Build.cs
--- a/Engine/Plugins/Experimental/Subsonic/Source/Subsonic/Subsonic.Build.cs
+++ b/Engine/Plugins/Experimental/Subsonic/Source/Subsonic/Subsonic.Build.cs
@@ -11,15 +11,18 @@ public class Subsonic : ModuleRules
 		PublicDependencyModuleNames.AddRange(
 			new string[]
 			{
+				"AudioMixerCore",
 				"Core",
+				"GameplayTags"
 			}
 		);
 			
 		PrivateDependencyModuleNames.AddRange(
 			new string[]
 			{
+				"AudioExtensions",
 				"CoreUObject",
-				"Engine",
+				"Engine"
 			}
 		);
 	}
diff --git a/Engine/Plugins/Experimental/Subsonic/Source/SubsonicEditor/Private/SubsonicAssetDefinitions.cpp b/Engine/Plugins/Experimental/Subsonic/Source/SubsonicEditor/Private/SubsonicAssetDefinitions.cpp
--- a/Engine/Plugins/Experimental/Subsonic/Source/SubsonicEditor/Private/SubsonicAssetDefinitions.cpp
+++ b/Engine/Plugins/Experimental/Subsonic/Source/SubsonicEditor/Private/SubsonicAssetDefinitions.cpp
@@ -3,69 +3,72 @@
 #include "SubsonicAssetDefinitions.h"
 
 #include "Styling/SlateStyleRegistry.h"
-#include "SubsonicEventHandler.h"
-#include "SubsonicEventHandlerEditor.h"
+#include "SubsonicEventCollection.h"
+#include "SubsonicEventCollectionEditor.h"
 
 #include UE_INLINE_GENERATED_CPP_BY_NAME(SubsonicAssetDefinitions)
 
 #define LOCTEXT_NAMESPACE "SubsonicEditor"
 
 
-FLinearColor UAssetDefinition_SubsonicEventHandler::GetAssetColor() const
+FLinearColor UAssetDefinition_SubsonicEventCollection::GetAssetColor() const
 {
  	if (const ISlateStyle* SubsonicStyle = FSlateStyleRegistry::FindSlateStyle("SubsonicStyle"))
  	{
- 		return SubsonicStyle->GetColor("SubsonicEventHandler.Color").ToFColorSRGB();
+ 		return SubsonicStyle->GetColor("SubsonicEventCollection.Color").ToFColorSRGB();
  	}
  
  	return FColor::White;
 }
 
-TSoftClassPtr<UObject> UAssetDefinition_SubsonicEventHandler::GetAssetClass() const
+TSoftClassPtr<UObject> UAssetDefinition_SubsonicEventCollection::GetAssetClass() const
 {
-	return USubsonicEventHandler::StaticClass();
+	return USubsonicEventCollection::StaticClass();
 }
 
-TConstArrayView<FAssetCategoryPath> UAssetDefinition_SubsonicEventHandler::GetAssetCategories() const
+TConstArrayView<FAssetCategoryPath> UAssetDefinition_SubsonicEventCollection::GetAssetCategories() const
 {
 	static const TArray<FAssetCategoryPath> Categories = { EAssetCategoryPaths::Audio / LOCTEXT("AssetSoundSubsonicSubMenu", "Subsonic") };
 	return Categories;
 }
 
-EAssetCommandResult UAssetDefinition_SubsonicEventHandler::OpenAssets(const FAssetOpenArgs& OpenArgs) const
+EAssetCommandResult UAssetDefinition_SubsonicEventCollection::OpenAssets(const FAssetOpenArgs& OpenArgs) const
 {
 	using namespace UE::Subsonic::Editor;
 
 	for (const FAssetData& AssetData : OpenArgs.Assets)
 	{
-		TSharedRef<FEventHandlerEditor> NewEditor = MakeShared<FEventHandlerEditor>();
-		NewEditor->Init(OpenArgs.GetToolkitMode(), OpenArgs.ToolkitHost, *CastChecked<USubsonicEventHandler>(AssetData.GetAsset()));
+		if (USubsonicEventCollection* Collection = Cast<USubsonicEventCollection>(AssetData.GetAsset()))
+		{
+			TSharedRef<FEventCollectionEditor> NewEditor = MakeShared<FEventCollectionEditor>();
+			NewEditor->Init(OpenArgs.GetToolkitMode(), OpenArgs.ToolkitHost, *Collection);
+		}
 	}
 
 	return EAssetCommandResult::Handled;
 }
 
-const FSlateBrush* UAssetDefinition_SubsonicEventHandler::GetThumbnailBrush(const FAssetData& InAssetData, const FName InClassName) const
+const FSlateBrush* UAssetDefinition_SubsonicEventCollection::GetThumbnailBrush(const FAssetData& InAssetData, const FName InClassName) const
 {
 	return Super::GetThumbnailBrush(InAssetData, InClassName);
 }
 
-const FSlateBrush* UAssetDefinition_SubsonicEventHandler::GetIconBrush(const FAssetData& InAssetData, const FName InClassName) const
+const FSlateBrush* UAssetDefinition_SubsonicEventCollection::GetIconBrush(const FAssetData& InAssetData, const FName InClassName) const
 {
 	return Super::GetIconBrush(InAssetData, InClassName);
 }
 
-bool UAssetDefinition_SubsonicEventHandler::GetThumbnailActionOverlay(const FAssetData& InAssetData, FAssetActionThumbnailOverlayInfo& OutActionOverlayInfo) const
+bool UAssetDefinition_SubsonicEventCollection::GetThumbnailActionOverlay(const FAssetData& InAssetData, FAssetActionThumbnailOverlayInfo& OutActionOverlayInfo) const
 {
 	return Super::GetThumbnailActionOverlay(InAssetData, OutActionOverlayInfo);
 }
 
-EAssetCommandResult UAssetDefinition_SubsonicEventHandler::ActivateAssets(const FAssetActivateArgs& ActivateArgs) const
+EAssetCommandResult UAssetDefinition_SubsonicEventCollection::ActivateAssets(const FAssetActivateArgs& ActivateArgs) const
 {
 	return Super::ActivateAssets(ActivateArgs);
 }
 
-void UAssetDefinition_SubsonicEventHandler::GetAssetActionButtonExtensions(const FAssetData& InAssetData, TArray<FAssetButtonActionExtension>& OutExtensions) const
+void UAssetDefinition_SubsonicEventCollection::GetAssetActionButtonExtensions(const FAssetData& InAssetData, TArray<FAssetButtonActionExtension>& OutExtensions) const
 {
 	Super::GetAssetActionButtonExtensions(InAssetData, OutExtensions);
 }
diff --git a/Engine/Plugins/Experimental/Subsonic/Source/SubsonicEditor/Private/SubsonicAssetDefinitions.h b/Engine/Plugins/Experimental/Subsonic/Source/SubsonicEditor/Private/SubsonicAssetDefinitions.h
--- a/Engine/Plugins/Experimental/Subsonic/Source/SubsonicEditor/Private/SubsonicAssetDefinitions.h
+++ b/Engine/Plugins/Experimental/Subsonic/Source/SubsonicEditor/Private/SubsonicAssetDefinitions.h
@@ -7,13 +7,13 @@
 
 
 UCLASS()
-class UAssetDefinition_SubsonicEventHandler : public UAssetDefinitionDefault_AudioDiffable
+class UAssetDefinition_SubsonicEventCollection : public UAssetDefinitionDefault_AudioDiffable
 {
 	GENERATED_BODY()
 
 public:
 	// UAssetDefinition Begin
-	virtual FText GetAssetDisplayName() const override { return NSLOCTEXT("AssetTypeActions", "AssetTypeActions_SubsonicEventHandler", "Subsonic Event Handler"); }
+	virtual FText GetAssetDisplayName() const override { return NSLOCTEXT("AssetTypeActions", "AssetTypeActions_SubsonicEventCollection", "Subsonic Event Collection"); }
 	virtual FLinearColor GetAssetColor() const override;
 	virtual TSoftClassPtr<UObject> GetAssetClass() const override;
 	virtual TConstArrayView<FAssetCategoryPath> GetAssetCategories() const override;
diff --git a/Engine/Plugins/Experimental/Subsonic/Source/SubsonicEditor/Private/SubsonicEventCollectionEditor.cpp b/Engine/Plugins/Experimental/Subsonic/Source/SubsonicEditor/Private/SubsonicEventCollectionEditor.cpp
new file mode 100644
--- /dev/null
+++ b/Engine/Plugins/Experimental/Subsonic/Source/SubsonicEditor/Private/SubsonicEventCollectionEditor.cpp
@@ -0,0 +1,198 @@
+// Copyright Epic Games, Inc. All Rights Reserved.
+
+#include "SubsonicEventCollectionEditor.h"
+
+#include "Framework/Application/SlateApplication.h"
+#include "Framework/Docking/TabManager.h"
+#include "PropertyEditorModule.h"
+#include "Modules/ModuleManager.h"
+#include "Widgets/Docking/SDockTab.h"
+#include "Styling/SlateStyleRegistry.h"
+#include "SubsonicEventCollection.h"
+
+#define LOCTEXT_NAMESPACE "SubsonicEditor"
+
+
+namespace UE::Subsonic::Editor
+{
+	namespace EventCollectionEditorPrivate
+	{
+		const FName TabName_Details = "SubsonicEventCollectionEditor_Details";
+	} // namespace EventCollectionEditorPrivate
+
+	const FName FEventCollectionEditor::EditorName = "SubsonicEventCollectionEditor";
+
+	void FEventCollectionEditor::RegisterTabSpawners(const TSharedRef<FTabManager>& InTabManager)
+	{
+		WorkspaceMenuCategory = InTabManager->AddLocalWorkspaceMenuCategory(LOCTEXT("WorkspaceMenu_SubsonicEventCollectionEditor", "Subsonic EventCollection Editor"));
+
+		FDetailsViewArgs Args;
+		Args.bHideSelectionTip = true;
+		Args.NotifyHook = this;
+
+		FPropertyEditorModule& PropertyModule = FModuleManager::LoadModuleChecked<FPropertyEditorModule>("PropertyEditor");
+		TSharedRef<IDetailsView> NewDetailsView = PropertyModule.CreateDetailView(Args);
+
+		InTabManager->RegisterTabSpawner(EventCollectionEditorPrivate::TabName_Details, FOnSpawnTab::CreateLambda([NewDetailsView](const FSpawnTabArgs& Args)
+			{
+				return SNew(SDockTab).Label(LOCTEXT("SubsonicEventCollectionDetailsTitle", "Details"))[ NewDetailsView ];
+			}))
+			.SetDisplayName(LOCTEXT("DetailsTab", "Details"))
+			.SetGroup(WorkspaceMenuCategory.ToSharedRef())
+			.SetIcon(FSlateIcon(FAppStyle::GetAppStyleSetName(), "LevelEditor.Tabs.Details"));
+
+		DetailsView = NewDetailsView;
+
+		FAssetEditorToolkit::RegisterTabSpawners(InTabManager);
+	}
+
+	void FEventCollectionEditor::UnregisterTabSpawners(const TSharedRef<FTabManager>& InTabManager)
+	{
+		FAssetEditorToolkit::UnregisterTabSpawners(InTabManager);
+
+		InTabManager->UnregisterTabSpawner(EventCollectionEditorPrivate::TabName_Details);
+	}
+
+	FEventCollectionEditor::~FEventCollectionEditor()
+	{
+		if (GEditor)
+		{
+			GEditor->UnregisterForUndo(this);
+		}
+	}
+
+	void FEventCollectionEditor::Init(const EToolkitMode::Type Mode, const TSharedPtr<IToolkitHost>& InitToolkitHost, USubsonicEventCollection& InEventCollection)
+	{
+		// Support undo/redo
+		InEventCollection.SetFlags(RF_Transactional);
+
+		if (GEditor)
+		{
+			GEditor->RegisterForUndo(this);
+		}
+
+		DetailsStack = FTabManager::NewStack()
+			->SetSizeCoefficient(0.50f)
+			->SetHideTabWell(false)
+			->AddTab(EventCollectionEditorPrivate::TabName_Details, ETabState::OpenedTab);
+
+		const TSharedRef<FTabManager::FLayout> StandaloneDefaultLayout = FTabManager::NewLayout("Standalone_SubsonicEventCollectionEditor_Layout_v1")
+			->AddArea
+			(
+				FTabManager::NewPrimaryArea()
+				->SetOrientation(Orient_Vertical)
+				->Split(FTabManager::NewSplitter()
+					->SetOrientation(Orient_Horizontal)
+					->Split
+					(
+						FTabManager::NewSplitter()
+						->SetSizeCoefficient(0.15f)
+						->SetOrientation(Orient_Vertical)
+						->Split
+						(
+							FTabManager::NewStack()
+							->SetSizeCoefficient(0.25f)
+							->SetHideTabWell(false)
+							->AddTab(EventCollectionEditorPrivate::TabName_Details, ETabState::OpenedTab)
+						)
+					)
+				)
+			);
+
+		{
+			constexpr bool bCreateDefaultStandaloneMenu = true;
+			constexpr bool bCreateDefaultToolbar = true;
+			constexpr bool bToolbarFocusable = false;
+			constexpr bool bUseSmallToolbarIcons = true;
+
+			FAssetEditorToolkit::InitAssetEditor(
+				Mode,
+				InitToolkitHost,
+				FName("SubsonicEventCollectionEditorApp"),
+				StandaloneDefaultLayout,
+				bCreateDefaultStandaloneMenu,
+				bCreateDefaultToolbar,
+				&InEventCollection,
+				bToolbarFocusable,
+				bUseSmallToolbarIcons);
+		}
+
+		if (DetailsView.IsValid())
+		{
+			DetailsView->SetObject(&InEventCollection);
+		}
+
+		RegenerateMenusAndToolbars();
+	}
+
+	USubsonicEventCollection* FEventCollectionEditor::GetEventCollection() const
+	{
+		return Cast<USubsonicEventCollection>(GetEditingObject());
+	}
+
+	FName FEventCollectionEditor::GetToolkitFName() const
+	{
+		return FEventCollectionEditor::EditorName;
+	}
+
+	FText FEventCollectionEditor::GetBaseToolkitName() const
+	{
+		return LOCTEXT("AppLabel", "Subsonic Event Collection Editor");
+	}
+
+	FString FEventCollectionEditor::GetWorldCentricTabPrefix() const
+	{
+		return LOCTEXT("WorldCentricTabPrefix", "Subsonic Event Collection ").ToString();
+	}
+
+	FLinearColor FEventCollectionEditor::GetWorldCentricTabColorScale() const
+	{
+		if (const ISlateStyle* Style = FSlateStyleRegistry::FindSlateStyle("SubsonicStyle"))
+		{
+			return Style->GetColor("SubsonicEventCollection.Color");
+		}
+
+		return FLinearColor(0.3f, 0.2f, 0.5f, 0.5f);
+	}
+
+	const FSlateBrush* FEventCollectionEditor::GetDefaultTabIcon() const
+	{
+		if (const ISlateStyle* Style = FSlateStyleRegistry::FindSlateStyle("SubsonicStyle"))
+		{
+			if (const FSlateBrush* Brush = Style->GetBrush("SubsonicEventCollectionEditor_Icon"))
+			{
+				return Brush;
+			}
+		}
+
+		if (const FSlateBrush* NoBrush = FAppStyle::GetBrush("NoBrush"))
+		{
+			return NoBrush;
+		}
+
+		static const FSlateBrush NullBrush;
+		return &NullBrush;
+	}
+
+	FLinearColor FEventCollectionEditor::GetDefaultTabColor() const
+	{
+		if (const ISlateStyle* Style = FSlateStyleRegistry::FindSlateStyle("SubsonicStyle"))
+		{
+			return Style->GetColor("SubsonicEventCollection.Color");
+		}
+
+		return FAssetEditorToolkit::GetDefaultTabColor();
+	}
+
+	FName FEventCollectionEditor::GetEditorName() const
+	{
+		return FEventCollectionEditor::EditorName;
+	}
+
+	void FEventCollectionEditor::PostUndo(bool bSuccess)
+	{
+		FSlateApplication::Get().DismissAllMenus();
+		FEditorUndoClient::PostUndo(bSuccess);
+	}
+}
+#undef LOCTEXT_NAMESPACE // SubsonicEditor
diff --git a/Engine/Plugins/Experimental/Subsonic/Source/SubsonicEditor/Private/SubsonicEventCollectionEditor.h b/Engine/Plugins/Experimental/Subsonic/Source/SubsonicEditor/Private/SubsonicEventCollectionEditor.h
new file mode 100644
--- /dev/null
+++ b/Engine/Plugins/Experimental/Subsonic/Source/SubsonicEditor/Private/SubsonicEventCollectionEditor.h
@@ -0,0 +1,55 @@
+// Copyright Epic Games, Inc. All Rights Reserved.
+#pragma once
+
+#include "SubsonicEventCollectionEditor.h"
+
+#include "EditorUndoClient.h"
+#include "Framework/Docking/TabManager.h"
+#include "IDetailsView.h"
+#include "Misc/NotifyHook.h"
+#include "Templates/SharedPointer.h"
+#include "Toolkits/AssetEditorToolkit.h"
+#include "Toolkits/IToolkitHost.h"
+
+
+// Forward Declarations
+class USubsonicEventCollection;
+
+
+namespace UE::Subsonic::Editor
+{
+	class FEventCollectionEditor : public FAssetEditorToolkit, public FNotifyHook, public FEditorUndoClient
+	{
+	public:
+		static const FName EditorName;
+
+		FEventCollectionEditor() = default;
+		virtual ~FEventCollectionEditor();
+
+		virtual void RegisterTabSpawners(const TSharedRef<FTabManager>& TabManager) override;
+		virtual void UnregisterTabSpawners(const TSharedRef<FTabManager>& TabManager) override;
+
+		void Init(const EToolkitMode::Type Mode, const TSharedPtr<IToolkitHost>& InitToolkitHost, USubsonicEventCollection& EventCollection);
+
+		USubsonicEventCollection* GetEventCollection() const;
+
+		/** IToolkit interface */
+		virtual FName GetToolkitFName() const override;
+		virtual FText GetBaseToolkitName() const override;
+		virtual FString GetWorldCentricTabPrefix() const override;
+		virtual FLinearColor GetWorldCentricTabColorScale() const override;
+		virtual const FSlateBrush* GetDefaultTabIcon() const override;
+		virtual FLinearColor GetDefaultTabColor() const override;
+
+		/** IAssetEditorInstance interface */
+		virtual FName GetEditorName() const override;
+
+		/** FEditorUndoClient Interface */
+		virtual void PostUndo(bool bSuccess) override;
+		virtual void PostRedo(bool bSuccess) override { PostUndo(bSuccess); }
+
+	private:
+		TSharedPtr<IDetailsView> DetailsView;
+		TSharedPtr<FTabManager::FStack> DetailsStack;
+	};
+} // namespace UE::Subsonic::Editor
diff --git a/Engine/Plugins/Experimental/Subsonic/Source/SubsonicEditor/Private/SubsonicEventCollectionFactory.cpp b/Engine/Plugins/Experimental/Subsonic/Source/SubsonicEditor/Private/SubsonicEventCollectionFactory.cpp
new file mode 100644
--- /dev/null
+++ b/Engine/Plugins/Experimental/Subsonic/Source/SubsonicEditor/Private/SubsonicEventCollectionFactory.cpp
@@ -0,0 +1,22 @@
+// Copyright Epic Games, Inc. All Rights Reserved.
+#include "SubsonicEventCollectionFactory.h"
+
+#include "SubsonicEventCollection.h"
+
+#include UE_INLINE_GENERATED_CPP_BY_NAME(SubsonicEventCollectionFactory)
+
+
+USubsonicEventCollectionFactory::USubsonicEventCollectionFactory(const FObjectInitializer& ObjectInitializer)
+	: Super(ObjectInitializer)
+{
+	bCreateNew = true;
+	bEditorImport = false;
+	bEditAfterNew = true;
+
+	SupportedClass = USubsonicEventCollection::StaticClass();
+}
+
+UObject* USubsonicEventCollectionFactory::FactoryCreateNew(UClass* InClass, UObject* InParent, FName InName, EObjectFlags InFlags, UObject* InContext, FFeedbackContext* InFeedbackContext)
+{
+	return NewObject<USubsonicEventCollection>(InParent, InClass, InName, InFlags);
+}
diff --git a/Engine/Plugins/Experimental/Subsonic/Source/SubsonicEditor/Private/SubsonicEventCollectionFactory.h b/Engine/Plugins/Experimental/Subsonic/Source/SubsonicEditor/Private/SubsonicEventCollectionFactory.h
new file mode 100644
--- /dev/null
+++ b/Engine/Plugins/Experimental/Subsonic/Source/SubsonicEditor/Private/SubsonicEventCollectionFactory.h
@@ -0,0 +1,19 @@
+// Copyright Epic Games, Inc. All Rights Reserved.
+
+#pragma once
+
+#include "Factories/Factory.h"
+#include "UObject/Object.h"
+
+#include "SubsonicEventCollectionFactory.generated.h"
+
+
+UCLASS()
+class USubsonicEventCollectionFactory : public UFactory
+{
+	GENERATED_UCLASS_BODY()
+
+	//~ Begin UFactory Interface
+	virtual UObject* FactoryCreateNew(UClass* InClass, UObject* InParent, FName InName, EObjectFlags InFlags, UObject* InContext, FFeedbackContext* InFeedbackContext) override;
+	//~ Begin UFactory Interface
+};
diff --git a/Engine/Plugins/Experimental/Subsonic/Source/SubsonicEditor/Private/SubsonicEventHandlerEditor.cpp b/Engine/Plugins/Experimental/Subsonic/Source/SubsonicEditor/Private/SubsonicEventHandlerEditor.cpp
deleted file mode 100644
diff --git a/Engine/Plugins/Experimental/Subsonic/Source/SubsonicEditor/Private/SubsonicEventHandlerEditor.h b/Engine/Plugins/Experimental/Subsonic/Source/SubsonicEditor/Private/SubsonicEventHandlerEditor.h
deleted file mode 100644
diff --git a/Engine/Plugins/Experimental/Subsonic/Source/SubsonicEditor/Private/SubsonicEventHandlerFactory.cpp b/Engine/Plugins/Experimental/Subsonic/Source/SubsonicEditor/Private/SubsonicEventHandlerFactory.cpp
deleted file mode 100644
diff --git a/Engine/Plugins/Experimental/Subsonic/Source/SubsonicEditor/Private/SubsonicEventHandlerFactory.h b/Engine/Plugins/Experimental/Subsonic/Source/SubsonicEditor/Private/SubsonicEventHandlerFactory.h
deleted file mode 100644
diff --git a/Engine/Plugins/Experimental/Subsonic/Source/SubsonicEditor/Private/SubsonicSlateStyle.h b/Engine/Plugins/Experimental/Subsonic/Source/SubsonicEditor/Private/SubsonicSlateStyle.h
--- a/Engine/Plugins/Experimental/Subsonic/Source/SubsonicEditor/Private/SubsonicSlateStyle.h
+++ b/Engine/Plugins/Experimental/Subsonic/Source/SubsonicEditor/Private/SubsonicSlateStyle.h
@@ -28,7 +28,7 @@ namespace UE::Subsonic::Editor
 
 			// Subsonic Editor
 			{
-				Set("SubsonicEventHandler.Color", FColor(156, 0, 255));
+				Set("SubsonicEventCollection.Color", FColor(156, 0, 255));
 			}
 
 			FSlateStyleRegistry::RegisterSlateStyle(*this);
