- Rename EventSystems to EventSubscribers for clarity/distinguishing from engine/audio device "Subsystems"
- Move AudioComponent Subscriber to private implementation
- Clean up pass on handles
- Make Executor handle independent from CollectionHandle
- Decouple executor handle construction from UObject/UniqueId
- Executor refactor to support non-uobject management
- Move all reflected Subsonic types to UE::Subsonic namespace
- Remove old test event code from Subsonic Subsystem

#rnx

#Synced-CL 50548439

diff --git a/Engine/Plugins/Experimental/Subsonic/Source/Subsonic/Private/CoreActions/SubsonicAction_AudioComponent.cpp b/Engine/Plugins/Experimental/Subsonic/Source/Subsonic/Private/CoreActions/SubsonicAction_AudioComponent.cpp
--- a/Engine/Plugins/Experimental/Subsonic/Source/Subsonic/Private/CoreActions/SubsonicAction_AudioComponent.cpp
+++ b/Engine/Plugins/Experimental/Subsonic/Source/Subsonic/Private/CoreActions/SubsonicAction_AudioComponent.cpp
@@ -9,10 +9,10 @@ namespace UE::Subsonic
 {
 	namespace AudioComponentPrivate
 	{
-		UAudioComponent* AddScopedComponent(const UE::Subsonic::FActionHandle& InHandle, ESubsonicExecutionScope Scope, FName Name)
+		UAudioComponent* AddScopedComponent(const UE::Subsonic::FActionScopeKey& InKey, ESubsonicExecutionScope Scope, FName Name)
 		{
 			UAudioComponent* Component = nullptr;
-			if (USubsonicAudioComponentSubsystem* AudioComponentSubsystem = InHandle.Event.Collection.FindDeviceSubsystem<USubsonicAudioComponentSubsystem>())
+			if (USubsonicAudioComponentSubscriber* AudioComponentSubsystem = InKey.FindDeviceSubsystem<USubsonicAudioComponentSubscriber>())
 			{
 				switch (Scope)
 				{
@@ -24,19 +24,19 @@ namespace UE::Subsonic
 
 					case ESubsonicExecutionScope::Collection:
 					{
-						Component = AudioComponentSubsystem->AddComponent(InHandle.Event.Collection, Name);
+						Component = AudioComponentSubsystem->AddComponent(InKey.ScopeHandle.Event.Collection, Name);
 					}
 					break;
 
 					case ESubsonicExecutionScope::Event:
 					{
-						Component = AudioComponentSubsystem->AddComponent(InHandle.Event, Name);
+						Component = AudioComponentSubsystem->AddComponent(InKey.ScopeHandle.Event, Name);
 					}
 					break;
 
 					case ESubsonicExecutionScope::Action:
 					{
-						Component = AudioComponentSubsystem->AddComponent(InHandle, Name);
+						Component = AudioComponentSubsystem->AddComponent(InKey.ScopeHandle, Name);
 					}
 					break;
 				}
@@ -45,12 +45,10 @@ namespace UE::Subsonic
 			return Component;
 		}
 
-		UAudioComponent* FindScopedComponent(const UE::Subsonic::FActionHandle& InHandle, ESubsonicExecutionScope Scope, FName Name)
+		UAudioComponent* FindScopedComponent(const UE::Subsonic::FActionScopeKey& InKey, ESubsonicExecutionScope Scope, FName Name)
 		{
-			using namespace UE::Subsonic;
-
 			UAudioComponent* Component = nullptr;
-			if (USubsonicAudioComponentSubsystem* AudioComponentSubsystem = InHandle.Event.Collection.FindDeviceSubsystem<USubsonicAudioComponentSubsystem>())
+			if (USubsonicAudioComponentSubscriber* AudioComponentSubsystem = InKey.FindDeviceSubsystem<USubsonicAudioComponentSubscriber>())
 			{
 				switch (Scope)
 					{
@@ -62,19 +60,19 @@ namespace UE::Subsonic
 
 					case ESubsonicExecutionScope::Collection:
 					{
-						Component = AudioComponentSubsystem->FindComponent(InHandle.Event.Collection, Name);
+						Component = AudioComponentSubsystem->FindComponent(InKey.ScopeHandle.Event.Collection, Name);
 					}
 					break;
 
 					case ESubsonicExecutionScope::Event:
 					{
-						Component = AudioComponentSubsystem->FindComponent(InHandle.Event, Name);
+						Component = AudioComponentSubsystem->FindComponent(InKey.ScopeHandle.Event, Name);
 					}
 					break;
 
 					case ESubsonicExecutionScope::Action:
 					{
-						Component = AudioComponentSubsystem->FindComponent(InHandle, Name);
+						Component = AudioComponentSubsystem->FindComponent(InKey.ScopeHandle, Name);
 					}
 					break;
 				}
@@ -83,12 +81,12 @@ namespace UE::Subsonic
 			return Component;
 		}
 
-		UAudioComponent* AccessComponent(const UE::Subsonic::FActionHandle& InHandle, ESubsonicAudioComponentAccess Access, ESubsonicExecutionScope Scope, FName Name)
+		UAudioComponent* AccessComponent(const UE::Subsonic::FActionScopeKey& InKey, ESubsonicAudioComponentAccess Access, ESubsonicExecutionScope Scope, FName Name)
 		{
 			UAudioComponent* Component = nullptr;
 			if (Access == ESubsonicAudioComponentAccess::Find || Access == ESubsonicAudioComponentAccess::FindOrAdd)
 			{
-				Component = FindScopedComponent(InHandle, Scope, Name);
+				Component = FindScopedComponent(InKey, Scope, Name);
 			}
 
 			bool bAddNewComponent = Access == ESubsonicAudioComponentAccess::Add;
@@ -104,80 +102,76 @@ namespace UE::Subsonic
 
 			if (bAddNewComponent)
 			{
-				Component = AddScopedComponent(InHandle, Scope, Name);
+				Component = AddScopedComponent(InKey, Scope, Name);
 			}
 
 			return Component;
 		}
 	} // namespace AudioComponentPrivate
-}
-
-void FSubsonicEventAction_AudioComponentModify::Execute(const UE::Subsonic::FActionHandle& InHandle) const
-{
-	using namespace UE::Subsonic;
 
-	if (UAudioComponent* AudioComponent = AudioComponentPrivate::AccessComponent(InHandle, Access, Scope, Name))
+	void FSubsonicEventAction_AudioComponentModify::Execute(const FSubsonicExecutor& InExecutor, const UE::Subsonic::FActionHandle& InHandle) const
 	{
-		for (const TInstancedStruct<FSubsonicEventAction_AudioComponentModifierBase>& Modifier : Modifiers)
+		const FActionScopeKey ScopeKey(InExecutor, InHandle);
+		if (UAudioComponent* AudioComponent = AudioComponentPrivate::AccessComponent(ScopeKey, Access, Scope, Name))
 		{
-			if (const FSubsonicEventAction_AudioComponentModifierBase* ModPtr = Modifier.GetPtr())
+			for (const TInstancedStruct<FSubsonicEventAction_AudioComponentModifierBase>& Modifier : Modifiers)
 			{
-				ModPtr->Execute(InHandle, *AudioComponent);
+				if (const FSubsonicEventAction_AudioComponentModifierBase* ModPtr = Modifier.GetPtr())
+				{
+					ModPtr->Execute(*AudioComponent);
+				}
 			}
 		}
 	}
-}
-
-void FSubsonicEventAction_AudioComponentPlay::Execute(const UE::Subsonic::FActionHandle& InHandle) const
-{
-	using namespace UE::Subsonic;
 
-	if (UAudioComponent* AudioComponent = AudioComponentPrivate::AccessComponent(InHandle, Access, Scope, Name))
+	void FSubsonicEventAction_AudioComponentPlay::Execute(const FSubsonicExecutor& InExecutor, const UE::Subsonic::FActionHandle& InHandle) const
 	{
-		AudioComponent->SetSound(Sound);
-		AudioComponent->Play();
+		const FActionScopeKey ScopeKey(InExecutor, InHandle);
+		if (UAudioComponent* AudioComponent = AudioComponentPrivate::AccessComponent(ScopeKey, Access, Scope, Name))
+		{
+			AudioComponent->SetSound(Sound);
+			AudioComponent->Play();
+		}
 	}
-}
-
-void FSubsonicEventAction_AudioComponentStop::Execute(const UE::Subsonic::FActionHandle& InHandle) const
-{
-	using namespace UE::Subsonic;
 
-	if (UAudioComponent* AudioComponent = AudioComponentPrivate::FindScopedComponent(InHandle, Scope, Name))
+	void FSubsonicEventAction_AudioComponentStop::Execute(const FSubsonicExecutor& InExecutor, const UE::Subsonic::FActionHandle& InHandle) const
 	{
-		AudioComponent->Stop();
+		const FActionScopeKey ScopeKey(InExecutor, InHandle);
+		if (UAudioComponent* AudioComponent = AudioComponentPrivate::FindScopedComponent(ScopeKey, Scope, Name))
+		{
+			AudioComponent->Stop();
+		}
 	}
-}
-
-void FSubsonicEventAction_AudioComponentModifier_Play::Execute(const UE::Subsonic::FActionHandle& InHandle, UAudioComponent& AudioComponent) const
-{
-	AudioComponent.Play(StartTime);
-}
 
+	void FSubsonicEventAction_AudioComponentModifier_Play::Execute(UAudioComponent& AudioComponent) const
+	{
+		AudioComponent.Play(StartTime);
+	}
 
-void FSubsonicEventAction_AudioComponentModifier_SetAttenuation::Execute(const UE::Subsonic::FActionHandle& InHandle, UAudioComponent& AudioComponent) const
-{
-	AudioComponent.SetAttenuationSettings(Attenuation);
-}
+	void FSubsonicEventAction_AudioComponentModifier_SetAttenuation::Execute(UAudioComponent& AudioComponent) const
+	{
+		AudioComponent.SetAttenuationSettings(Attenuation);
+	}
 
-void FSubsonicEventAction_AudioComponentModifier_SetConcurrency::Execute(const UE::Subsonic::FActionHandle& InHandle, UAudioComponent& AudioComponent) const
-{
-	AudioComponent.ConcurrencySet = Concurrency;
-}
+	void FSubsonicEventAction_AudioComponentModifier_SetConcurrency::Execute(UAudioComponent& AudioComponent) const
+	{
+		AudioComponent.ConcurrencySet = Concurrency;
+	}
 
-void FSubsonicEventAction_AudioComponentModifier_SetModulationRouting::Execute(const UE::Subsonic::FActionHandle& InHandle, UAudioComponent& AudioComponent) const
-{
-	TSet<USoundModulatorBase*> ModulatorSet;
-	Algo::Transform(Modulators, ModulatorSet, [](const TObjectPtr<USoundModulatorBase>& Modulator) { return Modulator.Get(); });
-	AudioComponent.SetModulationRouting(ModulatorSet, Destination, RoutingMethod);
-}
+	void FSubsonicEventAction_AudioComponentModifier_SetModulationRouting::Execute(UAudioComponent& AudioComponent) const
+	{
+		TSet<USoundModulatorBase*> ModulatorSet;
+		Algo::Transform(Modulators, ModulatorSet, [](const TObjectPtr<USoundModulatorBase>& Modulator) { return Modulator.Get(); });
+		AudioComponent.SetModulationRouting(ModulatorSet, Destination, RoutingMethod);
+	}
 
-void FSubsonicEventAction_AudioComponentModifier_SetSound::Execute(const UE::Subsonic::FActionHandle& InHandle, UAudioComponent& AudioComponent) const
-{
-	AudioComponent.SetSound(Sound);
-}
+	void FSubsonicEventAction_AudioComponentModifier_SetSound::Execute(UAudioComponent& AudioComponent) const
+	{
+		AudioComponent.SetSound(Sound);
+	}
 
-void FSubsonicEventAction_AudioComponentModifier_Stop::Execute(const UE::Subsonic::FActionHandle& InHandle, UAudioComponent& AudioComponent) const
-{
-	AudioComponent.Stop();
-}
\ No newline at end of file
+	void FSubsonicEventAction_AudioComponentModifier_Stop::Execute(UAudioComponent& AudioComponent) const
+	{
+		AudioComponent.Stop();
+	}
+} // namespace UE::Subsonic
\ No newline at end of file
diff --git a/Engine/Plugins/Experimental/Subsonic/Source/Subsonic/Private/CoreActions/SubsonicAction_AudioComponent.h b/Engine/Plugins/Experimental/Subsonic/Source/Subsonic/Private/CoreActions/SubsonicAction_AudioComponent.h
--- a/Engine/Plugins/Experimental/Subsonic/Source/Subsonic/Private/CoreActions/SubsonicAction_AudioComponent.h
+++ b/Engine/Plugins/Experimental/Subsonic/Source/Subsonic/Private/CoreActions/SubsonicAction_AudioComponent.h
@@ -1,7 +1,7 @@
 // Copyright Epic Games, Inc. All Rights Reserved.
 #pragma once
 
-#include "CoreEventSystems/SubsonicAudioComponentSubsystem.h"
+#include "CoreEventSubscribers/AudioComponentEventSubscriber.h"
 #include "Misc/CoreMiscDefines.h"
 #include "Sound/SoundBase.h"
 #include "Sound/SoundConcurrency.h"
@@ -20,192 +20,195 @@ class USoundModulatorBase;
 struct FSubsonicEventAction_AudioComponentModifierBase;
 
 
-UENUM(BlueprintType)
-enum class ESubsonicAudioComponentAccess : uint8
+namespace UE::Subsonic
 {
-	// Adds new component, releasing any existing component.
-	// Stops existing component if sound is not one-shot.
-	Add,
+	UENUM(BlueprintType)
+	enum class ESubsonicAudioComponentAccess : uint8
+	{
+		// Adds new component, releasing any existing component.
+		// Stops existing component if sound is not one-shot.
+		Add,
 
-	// Finds existing component, or creates a new one.
-	FindOrAdd,
+		// Finds existing component, or creates a new one.
+		FindOrAdd,
 
-	// Ignore request and if component doesn't already exist.
-	Find
-};
+		// Ignore request and if component doesn't already exist.
+		Find
+	};
 
-USTRUCT(MinimalAPI, BlueprintType, DisplayName = "Modify Audio Component")
-struct FSubsonicEventAction_AudioComponentModify : public FSubsonicEventActionBase
-{
-	GENERATED_BODY()
+	USTRUCT(MinimalAPI, BlueprintType, DisplayName = "Modify Audio Component")
+	struct FSubsonicEventAction_AudioComponentModify : public FSubsonicEventActionBase
+	{
+		GENERATED_BODY()
 
-public:
-	virtual ~FSubsonicEventAction_AudioComponentModify() = default;
+	public:
+		virtual ~FSubsonicEventAction_AudioComponentModify() = default;
 
-	virtual void Execute(const UE::Subsonic::FActionHandle& InHandle) const override;
+		virtual void Execute(const FSubsonicExecutor& InExecutor, const UE::Subsonic::FActionHandle& InHandle) const override;
 
-	// Addressable name of audio component instance.
-	UPROPERTY(EditAnywhere, Category = "Subsonic")
-	FName Name = "Sound01";
+		// Addressable name of audio component instance.
+		UPROPERTY(EditAnywhere, Category = "Subsonic")
+		FName Name = "Sound01";
 
-	// Scope of component accessor.  If set to event, provided name is used
-	// to address the same component within the event, collection the same component
-	// within the collection, etc.
-	UPROPERTY(EditAnywhere, Category = "Subsonic")
-	ESubsonicExecutionScope Scope = ESubsonicExecutionScope::Collection;
+		// Scope of component accessor.  If set to event, provided name is used
+		// to address the same component within the event, collection the same component
+		// within the collection, etc.
+		UPROPERTY(EditAnywhere, Category = "Subsonic")
+		ESubsonicExecutionScope Scope = ESubsonicExecutionScope::Collection;
 
-	UPROPERTY(EditAnywhere, Category = "Subsonic")
-	ESubsonicAudioComponentAccess Access = ESubsonicAudioComponentAccess::FindOrAdd;
+		UPROPERTY(EditAnywhere, Category = "Subsonic")
+		ESubsonicAudioComponentAccess Access = ESubsonicAudioComponentAccess::FindOrAdd;
 
-	UPROPERTY(EditAnywhere, Category = "Subsonic")
-	TArray<TInstancedStruct<FSubsonicEventAction_AudioComponentModifierBase>> Modifiers;
-};
+		UPROPERTY(EditAnywhere, Category = "Subsonic")
+		TArray<TInstancedStruct<FSubsonicEventAction_AudioComponentModifierBase>> Modifiers;
+	};
 
-USTRUCT(MinimalAPI, BlueprintType, DisplayName = "Play Sound")
-struct FSubsonicEventAction_AudioComponentPlay : public FSubsonicEventActionBase
-{
-	GENERATED_BODY()
+	USTRUCT(MinimalAPI, BlueprintType, DisplayName = "Play Sound")
+	struct FSubsonicEventAction_AudioComponentPlay : public FSubsonicEventActionBase
+	{
+		GENERATED_BODY()
 
-	virtual ~FSubsonicEventAction_AudioComponentPlay() = default;
+		virtual ~FSubsonicEventAction_AudioComponentPlay() = default;
 
-	virtual void Execute(const UE::Subsonic::FActionHandle& InHandle) const override;
+		virtual void Execute(const FSubsonicExecutor& InExecutor, const UE::Subsonic::FActionHandle& InHandle) const override;
 
-public:
-	// Addressable name of audio component instance.
-	UPROPERTY(EditAnywhere, Category = "Subsonic")
-	FName Name = "Sound01";
+	public:
+		// Addressable name of audio component instance.
+		UPROPERTY(EditAnywhere, Category = "Subsonic")
+		FName Name = "Sound01";
 
-	// Sound to play.
-	UPROPERTY(EditAnywhere, Category = "Subsonic")
-	TObjectPtr<USoundBase> Sound;
+		// Sound to play.
+		UPROPERTY(EditAnywhere, Category = "Subsonic")
+		TObjectPtr<USoundBase> Sound;
 
-	// Scope of component accessor.  If set to event, provided name is used
-	// to address the same component within the event, collection the same component
-	// within the collection, etc.
-	UPROPERTY(EditAnywhere, Category = "Subsonic")
-	ESubsonicExecutionScope Scope = ESubsonicExecutionScope::Collection;
+		// Scope of component accessor.  If set to event, provided name is used
+		// to address the same component within the event, collection the same component
+		// within the collection, etc.
+		UPROPERTY(EditAnywhere, Category = "Subsonic")
+		ESubsonicExecutionScope Scope = ESubsonicExecutionScope::Collection;
 
-	UPROPERTY(EditAnywhere, Category = "Subsonic")
-	ESubsonicAudioComponentAccess Access = ESubsonicAudioComponentAccess::FindOrAdd;
-};
+		UPROPERTY(EditAnywhere, Category = "Subsonic")
+		ESubsonicAudioComponentAccess Access = ESubsonicAudioComponentAccess::FindOrAdd;
+	};
 
-USTRUCT(MinimalAPI, BlueprintType, DisplayName = "Stop Sound")
-struct FSubsonicEventAction_AudioComponentStop : public FSubsonicEventActionBase
-{
-	GENERATED_BODY()
+	USTRUCT(MinimalAPI, BlueprintType, DisplayName = "Stop Sound")
+	struct FSubsonicEventAction_AudioComponentStop : public FSubsonicEventActionBase
+	{
+		GENERATED_BODY()
 
-	virtual ~FSubsonicEventAction_AudioComponentStop() = default;
+		virtual ~FSubsonicEventAction_AudioComponentStop() = default;
 
-	virtual void Execute(const UE::Subsonic::FActionHandle& InHandle) const override;
+		virtual void Execute(const FSubsonicExecutor& InExecutor, const UE::Subsonic::FActionHandle& InHandle) const override;
 
-public:
-	// Addressable name of component/sound instance.
-	UPROPERTY(EditAnywhere, Category = "Subsonic")
-	FName Name = "Sound01";
+	public:
+		// Addressable name of component/sound instance.
+		UPROPERTY(EditAnywhere, Category = "Subsonic")
+		FName Name = "Sound01";
 
-	// Scope of component to play.  If set to event, provided name is used
-	// to address the existing component within the event, collection the existing component
-	// within the collection, etc.
-	UPROPERTY(EditAnywhere, Category = "Subsonic")
-	ESubsonicExecutionScope Scope = ESubsonicExecutionScope::Collection;
-};
+		// Scope of component to play.  If set to event, provided name is used
+		// to address the existing component within the event, collection the existing component
+		// within the collection, etc.
+		UPROPERTY(EditAnywhere, Category = "Subsonic")
+		ESubsonicExecutionScope Scope = ESubsonicExecutionScope::Collection;
+	};
 
-USTRUCT(MinimalAPI, BlueprintType, meta = (Abstract, Hidden))
-struct FSubsonicEventAction_AudioComponentModifierBase
-{
-	GENERATED_BODY()
+	USTRUCT(MinimalAPI, BlueprintType, meta = (Abstract, Hidden))
+	struct FSubsonicEventAction_AudioComponentModifierBase
+	{
+		GENERATED_BODY()
 
-public:
-	virtual ~FSubsonicEventAction_AudioComponentModifierBase() = default;
+	public:
+		virtual ~FSubsonicEventAction_AudioComponentModifierBase() = default;
 
-	virtual void Execute(const UE::Subsonic::FActionHandle& InHandle, UAudioComponent& AudioComponent) const { }
-};
+		virtual void Execute(UAudioComponent& AudioComponent) const { }
+	};
 
-USTRUCT(MinimalAPI, BlueprintType, DisplayName = "Play")
-struct FSubsonicEventAction_AudioComponentModifier_Play : public FSubsonicEventAction_AudioComponentModifierBase
-{
-	GENERATED_BODY()
+	USTRUCT(MinimalAPI, BlueprintType, DisplayName = "Play")
+	struct FSubsonicEventAction_AudioComponentModifier_Play : public FSubsonicEventAction_AudioComponentModifierBase
+	{
+		GENERATED_BODY()
 
-public:
-	virtual ~FSubsonicEventAction_AudioComponentModifier_Play() = default;
+	public:
+		virtual ~FSubsonicEventAction_AudioComponentModifier_Play() = default;
 
-	virtual void Execute(const UE::Subsonic::FActionHandle& InHandle, UAudioComponent& AudioComponent) const override;
+		virtual void Execute(UAudioComponent& AudioComponent) const override;
 
-	// Sound to play.
-	UPROPERTY(EditAnywhere, Category = "Subsonic")
-	float StartTime = 0.0f;
-};
+		// Sound to play.
+		UPROPERTY(EditAnywhere, Category = "Subsonic")
+		float StartTime = 0.0f;
+	};
 
-USTRUCT(MinimalAPI, BlueprintType, DisplayName = "Set Attenuation")
-struct FSubsonicEventAction_AudioComponentModifier_SetAttenuation : public FSubsonicEventAction_AudioComponentModifierBase
-{
-	GENERATED_BODY()
+	USTRUCT(MinimalAPI, BlueprintType, DisplayName = "Set Attenuation")
+	struct FSubsonicEventAction_AudioComponentModifier_SetAttenuation : public FSubsonicEventAction_AudioComponentModifierBase
+	{
+		GENERATED_BODY()
 
-public:
-	virtual ~FSubsonicEventAction_AudioComponentModifier_SetAttenuation() = default;
+	public:
+		virtual ~FSubsonicEventAction_AudioComponentModifier_SetAttenuation() = default;
 
-	virtual void Execute(const UE::Subsonic::FActionHandle& InHandle, UAudioComponent& AudioComponent) const override;
+		virtual void Execute(UAudioComponent& AudioComponent) const override;
 
-	UPROPERTY(EditAnywhere, Category = "Subsonic")
-	TObjectPtr<USoundAttenuation> Attenuation;
-};
+		UPROPERTY(EditAnywhere, Category = "Subsonic")
+		TObjectPtr<USoundAttenuation> Attenuation;
+	};
 
-USTRUCT(MinimalAPI, BlueprintType, DisplayName = "Set Concurrency")
-struct FSubsonicEventAction_AudioComponentModifier_SetConcurrency : public FSubsonicEventAction_AudioComponentModifierBase
-{
-	GENERATED_BODY()
+	USTRUCT(MinimalAPI, BlueprintType, DisplayName = "Set Concurrency")
+	struct FSubsonicEventAction_AudioComponentModifier_SetConcurrency : public FSubsonicEventAction_AudioComponentModifierBase
+	{
+		GENERATED_BODY()
 
-public:
-	virtual ~FSubsonicEventAction_AudioComponentModifier_SetConcurrency() = default;
+	public:
+		virtual ~FSubsonicEventAction_AudioComponentModifier_SetConcurrency() = default;
 
-	virtual void Execute(const UE::Subsonic::FActionHandle& InHandle, UAudioComponent& AudioComponent) const override;
+		virtual void Execute(UAudioComponent& AudioComponent) const override;
 
-	UPROPERTY(EditAnywhere, Category = "Subsonic")
-	TSet<TObjectPtr<USoundConcurrency>> Concurrency;
-};
+		UPROPERTY(EditAnywhere, Category = "Subsonic")
+		TSet<TObjectPtr<USoundConcurrency>> Concurrency;
+	};
 
-USTRUCT(MinimalAPI, BlueprintType, DisplayName = "Set Modulation Routing")
-struct FSubsonicEventAction_AudioComponentModifier_SetModulationRouting: public FSubsonicEventAction_AudioComponentModifierBase
-{
-	GENERATED_BODY()
+	USTRUCT(MinimalAPI, BlueprintType, DisplayName = "Set Modulation Routing")
+	struct FSubsonicEventAction_AudioComponentModifier_SetModulationRouting: public FSubsonicEventAction_AudioComponentModifierBase
+	{
+		GENERATED_BODY()
 
-public:
-	virtual ~FSubsonicEventAction_AudioComponentModifier_SetModulationRouting() = default;
+	public:
+		virtual ~FSubsonicEventAction_AudioComponentModifier_SetModulationRouting() = default;
 
-	virtual void Execute(const UE::Subsonic::FActionHandle& InHandle, UAudioComponent& AudioComponent) const override;
+		virtual void Execute(UAudioComponent& AudioComponent) const override;
 
-	UPROPERTY(EditAnywhere, Category = "Subsonic")
-	TSet<TObjectPtr<USoundModulatorBase>> Modulators;
+		UPROPERTY(EditAnywhere, Category = "Subsonic")
+		TSet<TObjectPtr<USoundModulatorBase>> Modulators;
 
-	UPROPERTY(EditAnywhere, Category = "Subsonic")
-	EModulationDestination Destination = EModulationDestination::Volume;
+		UPROPERTY(EditAnywhere, Category = "Subsonic")
+		EModulationDestination Destination = EModulationDestination::Volume;
 	
-	UPROPERTY(EditAnywhere, Category = "Subsonic")
-	EModulationRouting RoutingMethod = EModulationRouting::Inherit;
-};
+		UPROPERTY(EditAnywhere, Category = "Subsonic")
+		EModulationRouting RoutingMethod = EModulationRouting::Inherit;
+	};
 
-USTRUCT(MinimalAPI, BlueprintType, DisplayName = "Set Sound")
-struct FSubsonicEventAction_AudioComponentModifier_SetSound : public FSubsonicEventAction_AudioComponentModifierBase
-{
-	GENERATED_BODY()
+	USTRUCT(MinimalAPI, BlueprintType, DisplayName = "Set Sound")
+	struct FSubsonicEventAction_AudioComponentModifier_SetSound : public FSubsonicEventAction_AudioComponentModifierBase
+	{
+		GENERATED_BODY()
 
-public:
-	virtual ~FSubsonicEventAction_AudioComponentModifier_SetSound() = default;
+	public:
+		virtual ~FSubsonicEventAction_AudioComponentModifier_SetSound() = default;
 
-	virtual void Execute(const UE::Subsonic::FActionHandle& InHandle, UAudioComponent& AudioComponent) const override;
+		virtual void Execute(UAudioComponent& AudioComponent) const override;
 
-	// Sound to play.
-	UPROPERTY(EditAnywhere, Category = "Subsonic")
-	TObjectPtr<USoundBase> Sound;
-};
+		// Sound to play.
+		UPROPERTY(EditAnywhere, Category = "Subsonic")
+		TObjectPtr<USoundBase> Sound;
+	};
 
-USTRUCT(MinimalAPI, BlueprintType, DisplayName = "Stop")
-struct FSubsonicEventAction_AudioComponentModifier_Stop : public FSubsonicEventAction_AudioComponentModifierBase
-{
-	GENERATED_BODY()
+	USTRUCT(MinimalAPI, BlueprintType, DisplayName = "Stop")
+	struct FSubsonicEventAction_AudioComponentModifier_Stop : public FSubsonicEventAction_AudioComponentModifierBase
+	{
+		GENERATED_BODY()
 
-public:
-	virtual ~FSubsonicEventAction_AudioComponentModifier_Stop() = default;
+	public:
+		virtual ~FSubsonicEventAction_AudioComponentModifier_Stop() = default;
 
-	virtual void Execute(const UE::Subsonic::FActionHandle& InHandle, UAudioComponent& AudioComponent) const override;
-};
+		virtual void Execute(UAudioComponent& AudioComponent) const override;
+	};
+} // namespace UE::Subsonic
\ No newline at end of file
diff --git a/Engine/Plugins/Experimental/Subsonic/Source/Subsonic/Private/CoreEventSubscribers/AudioComponentEventSubscriber.cpp b/Engine/Plugins/Experimental/Subsonic/Source/Subsonic/Private/CoreEventSubscribers/AudioComponentEventSubscriber.cpp
new file mode 100644
--- /dev/null
+++ b/Engine/Plugins/Experimental/Subsonic/Source/Subsonic/Private/CoreEventSubscribers/AudioComponentEventSubscriber.cpp
@@ -0,0 +1,567 @@
+// Copyright Epic Games, Inc. All Rights Reserved.
+
+#include "CoreEventSubscribers/AudioComponentEventSubscriber.h"
+
+#include "Algo/Transform.h"
+#include "AudioDeviceManager.h"
+#include "AudioDevice.h"
+#include "Engine/Engine.h"
+#include "SubsonicExecutor.h"
+#include "SubsonicHandles.h"
+#include "SubsonicLog.h"
+#include "SubsonicSubsystem.h"
+#include "UObject/GCObject.h"
+#include "UObject/UObjectGlobals.h"
+
+#define SUBSONIC_ASSERT_ON_STALE_COLLECTION_EXECUTORS (1 && DO_CHECK)
+
+namespace UE::Subsonic
+{
+	namespace AudioComponentEventSubscriberPrivate
+	{
+		struct FInstanceAudioComponents : public FGCObject
+		{
+			// Currently addressable component via handle
+			TObjectPtr<UAudioComponent> ActiveComponent;
+
+			// Components no longer addressable from actions but actively
+			// playing. Continually tracked here to avoid GC while finishing.
+			TArray<TObjectPtr<UAudioComponent>> ReleasedComponents;
+
+			virtual void AddReferencedObjects(FReferenceCollector& Collector) override
+			{
+				Collector.AddReferencedObject(ActiveComponent);
+				Collector.AddStableReferenceArray(&ReleasedComponents);
+			}
+
+			virtual FString GetReferencerName() const override
+			{
+				return "UE::Subsonic::InstanceAudioComponents";
+			}
+
+			bool IsEmpty() const
+			{
+				return !ActiveComponent.Get() && ReleasedComponents.IsEmpty();
+			}
+
+			void ReleaseComponent()
+			{
+				if (ActiveComponent.Get() && ActiveComponent->IsPlaying())
+				{
+					ReleasedComponents.Add(ActiveComponent);
+				}
+				ActiveComponent = nullptr;
+			}
+
+			void StopAll() const
+			{
+				if (UAudioComponent* Component = ActiveComponent.Get())
+				{
+					Component->Stop();
+				}
+
+				for (const TObjectPtr<UAudioComponent>& ReleasedComponent : ReleasedComponents)
+				{
+					if (UAudioComponent* Component = ReleasedComponent.Get())
+					{
+						Component->Stop();
+					}
+				}
+			}
+		};
+
+		using FInstanceAudioComponentsPtr = TUniquePtr<FInstanceAudioComponents>;
+		using FAudioComponentsCache = TMap<FName, FInstanceAudioComponentsPtr>;
+
+		class FSubscriberImpl : public TSharedFromThis<FSubscriberImpl>
+		{
+			FAudioComponentsCache GlobalComponents;
+
+			TMap<FCollectionHandle, FAudioComponentsCache> CollectionComponents;
+			TMap<FEventHandle, FAudioComponentsCache> EventComponents;
+			TMap<FActionHandle, FAudioComponentsCache> ActionComponents;
+
+			TMap<uint32, FCollectionHandle> CollectionExecutors;
+
+			static UAudioComponent* CreateComponent(UObject& Parent, Audio::FDeviceId DeviceId, FName Name)
+			{
+				const FName ComponentName = MakeUniqueObjectName(&Parent, UAudioComponent::StaticClass(), Name);
+				UAudioComponent* NewComponent = NewObject<UAudioComponent>(&Parent, ComponentName);
+
+				check(NewComponent);
+				NewComponent->AudioDeviceID = DeviceId;
+				NewComponent->bAutoDestroy = true;
+				NewComponent->bIsUISound = true;
+				NewComponent->bAllowSpatialization = false;
+				NewComponent->bReverb = false;
+				NewComponent->bCenterChannelOnly = false;
+
+				return NewComponent;
+			}
+
+			static UAudioComponent* AddCachedComponent(
+				UObject& Parent,
+				Audio::FDeviceId DeviceId,
+				const FInstanceAudioComponentsPtr& InstanceComponents,
+				FName Name,
+				bool bReleaseExisting,
+				bool& bOutStoppedComponent)
+			{
+				bOutStoppedComponent = false;
+
+				if (bReleaseExisting)
+				{
+					if (UAudioComponent* Component = InstanceComponents->ActiveComponent.Get(); Component && Component->IsPlaying())
+					{
+						if (USoundBase* Sound = Component->GetSound(); Sound && !Sound->IsOneShot())
+						{
+							bOutStoppedComponent = true;
+							Component->Stop();
+						}
+					}
+
+					InstanceComponents->ReleaseComponent();
+				}
+
+				TObjectPtr<UAudioComponent>& ComponentRef = InstanceComponents->ActiveComponent;
+				if (!ComponentRef)
+				{
+					ComponentRef = CreateComponent(Parent, DeviceId, Name);
+				}
+
+				return ComponentRef.Get();
+			}
+
+			template <typename HandleType>
+			UAudioComponent* AddCachedComponent(
+				UObject& Parent,
+				Audio::FDeviceId DeviceId,
+				const HandleType& InHandle,
+				TMap<HandleType, FAudioComponentsCache>& OutScopeComponentMap,
+				FName Name,
+				bool bReleaseExisting)
+			{
+				UAudioComponent* CacheComponent = nullptr;
+
+				FAudioComponentsCache& Cache = OutScopeComponentMap.FindOrAdd(InHandle);
+				FInstanceAudioComponentsPtr& InstanceComponents = Cache.FindOrAdd(Name);
+				if (!InstanceComponents.IsValid())
+				{
+					InstanceComponents = MakeUnique<FInstanceAudioComponents>();
+				}
+
+				bool bStoppedComponent = false;
+				CacheComponent = AddCachedComponent(Parent, DeviceId, InstanceComponents, Name, bReleaseExisting, bStoppedComponent);
+				if (bStoppedComponent)
+				{
+					UE_LOG(LogSubsonic, Warning,
+						TEXT("Stopping released sound with handle '%s': "
+							"Sound is not one shot and continued play would result in leaked playback."),
+						*InHandle.ToString())
+				}
+				auto RemoveComponentLambda = [FinishHandle = InHandle, &OutScopeComponentMap, Name](UAudioComponent* FinishingComponent)
+				{
+					if (FAudioComponentsCache* Cache = OutScopeComponentMap.Find(FinishHandle))
+					{
+						if (FInstanceAudioComponentsPtr* CacheComponentsPtr = Cache->Find(Name); CacheComponentsPtr && CacheComponentsPtr->IsValid())
+						{
+							FInstanceAudioComponentsPtr& CacheComponents = *CacheComponentsPtr;
+							if (CacheComponents->ActiveComponent.Get() == FinishingComponent)
+							{
+								CacheComponents->ActiveComponent = nullptr;
+							}
+
+							CacheComponents->ReleasedComponents.RemoveAllSwap([&FinishingComponent](const TObjectPtr<UAudioComponent>& Released)
+							{
+								return Released.Get() == FinishingComponent;
+							});
+
+							if (CacheComponents->IsEmpty())
+							{
+								Cache->Remove(Name);
+							}
+						}
+						if (Cache->IsEmpty())
+						{
+							OutScopeComponentMap.Remove(FinishHandle);
+						}
+					}
+				};
+
+				CacheComponent->OnAudioFinishedNative.AddSPLambda(this, RemoveComponentLambda);
+				return CacheComponent;
+			}
+
+			template <typename HandleType>
+			UAudioComponent* FindCacheComponent(const TMap<HandleType, FAudioComponentsCache>& ScopeComponentMap, const HandleType& InHandle, FName Name) const
+			{
+				if (const FAudioComponentsCache* Cache = ScopeComponentMap.Find(InHandle))
+				{
+					if (const FInstanceAudioComponentsPtr* CacheComponentsPtr = Cache->Find(Name); CacheComponentsPtr && CacheComponentsPtr->IsValid())
+					{
+						return (*CacheComponentsPtr)->ActiveComponent.Get();
+					}
+				}
+
+				return nullptr;
+			}
+
+			template <typename HandleType>
+			bool ContainsCacheComponent(const TMap<HandleType, FAudioComponentsCache>& ScopeComponentMap, const HandleType& InHandle, FName Name) const
+			{
+				if (const FAudioComponentsCache* Cache = ScopeComponentMap.Find(InHandle))
+				{
+					if (const FInstanceAudioComponentsPtr* CacheComponentsPtr = Cache->Find(Name))
+					{
+						return CacheComponentsPtr->IsValid();
+					}
+				}
+
+				return false;
+			}
+
+			template <typename HandleType>
+			void RemoveCacheComponent(TMap<HandleType, FAudioComponentsCache>& ScopeComponentMap, const HandleType& InHandle, FName Name)
+			{
+				if (FAudioComponentsCache* Cache = ScopeComponentMap.Find(InHandle))
+				{
+					FInstanceAudioComponentsPtr* InstanceComponents = Cache->Find(Name);
+					if (InstanceComponents && InstanceComponents->IsValid())
+					{
+						(*InstanceComponents)->StopAll();
+					}
+
+					Cache->Remove(Name);
+					if (Cache->IsEmpty())
+					{
+						ScopeComponentMap.Remove(InHandle);
+					}
+				}
+			}
+
+			template <typename THandleType>
+			void ResetScopeComponents(TMap<THandleType, FAudioComponentsCache>& ScopeComponentMap)
+			{
+				for (const TPair<THandleType, FAudioComponentsCache>& HandleCachePair : ScopeComponentMap)
+				{
+					for (const TPair<FName, FInstanceAudioComponentsPtr>& Pair : HandleCachePair.Value)
+					{
+						const FInstanceAudioComponentsPtr& InstanceComponents = Pair.Value;
+						check(InstanceComponents.IsValid());
+						InstanceComponents->StopAll();
+					}
+				}
+				ScopeComponentMap.Reset();
+			};
+
+		public:
+			void Deinitialize()
+			{
+				for (const TPair<FName, FInstanceAudioComponentsPtr>& Pair : GlobalComponents)
+				{
+					const FInstanceAudioComponentsPtr& InstanceComponents = Pair.Value;
+					check(InstanceComponents.IsValid());
+					InstanceComponents->StopAll();
+				}
+				GlobalComponents.Empty();
+
+				ResetScopeComponents(CollectionComponents);
+				ResetScopeComponents(EventComponents);
+				ResetScopeComponents(ActionComponents);
+			}
+
+			UAudioComponent* AddComponent(UObject& Parent, Audio::FDeviceId DeviceId, FName Name, bool bReleaseExisting)
+			{
+				FInstanceAudioComponentsPtr& InstanceComponents = GlobalComponents.FindOrAdd(Name);
+				if (!InstanceComponents.IsValid())
+				{
+					InstanceComponents = MakeUnique<FInstanceAudioComponents>();
+				}
+
+				bool bStoppedComponent = false;
+				if (UAudioComponent* CacheComponent = AddCachedComponent(Parent, DeviceId, InstanceComponents, Name, bReleaseExisting, bStoppedComponent))
+				{
+					if (bStoppedComponent)
+					{
+						UE_LOG(LogSubsonic, Warning,
+							TEXT("Stopping released global sound '%s': "
+								"Sound is not one shot and continued play would result in leaked playback."),
+							*Name.ToString())
+					}
+
+					auto RemoveComponentLambda = [this, Name](UAudioComponent* FinishingComponent)
+					{
+						if (FInstanceAudioComponentsPtr* Components = GlobalComponents.Find(Name); Components && Components->IsValid())
+						{
+							FInstanceAudioComponents& ComponentsRef = (*Components->Get());
+							if (ComponentsRef.ActiveComponent.Get() == FinishingComponent)
+							{
+								ComponentsRef.ActiveComponent = nullptr;
+							}
+
+							ComponentsRef.ReleasedComponents.RemoveAllSwap([&FinishingComponent](const TObjectPtr<UAudioComponent>& Released)
+							{
+								return Released.Get() == FinishingComponent;
+							});
+
+							if (ComponentsRef.IsEmpty())
+							{
+								GlobalComponents.Remove(Name);
+							}
+						}
+					};
+					CacheComponent->OnAudioFinishedNative.AddSPLambda(this, RemoveComponentLambda);
+					return CacheComponent;
+				}
+
+				return nullptr;
+			}
+
+			UAudioComponent* AddComponent(UObject& Parent, Audio::FDeviceId DeviceId, const FCollectionHandle& InHandle, FName Name, bool bReleaseExisting)
+			{
+				return AddCachedComponent(Parent, DeviceId, InHandle, CollectionComponents, Name, bReleaseExisting);
+			}
+
+			UAudioComponent* AddComponent(UObject& Parent, Audio::FDeviceId DeviceId, const FEventHandle& InHandle, FName Name, bool bReleaseExisting)
+			{
+				return AddCachedComponent(Parent, DeviceId, InHandle, EventComponents, Name, bReleaseExisting);
+			}
+
+			UAudioComponent* AddComponent(UObject& Parent, Audio::FDeviceId DeviceId, const FActionHandle& InHandle, FName Name, bool bReleaseExisting)
+			{
+				return AddCachedComponent(Parent, DeviceId, InHandle, ActionComponents, Name, bReleaseExisting);
+			}
+
+			bool ContainsComponent(FName Name) const
+			{
+				return GlobalComponents.Contains(Name);
+			}
+
+			bool ContainsComponent(const FCollectionHandle& InHandle, FName Name) const
+			{
+				return ContainsCacheComponent(CollectionComponents, InHandle, Name);
+			}
+
+			bool ContainsComponent(const FEventHandle& InHandle, FName Name) const
+			{
+				return ContainsCacheComponent(EventComponents, InHandle, Name);
+			}
+
+			bool ContainsComponent(const FActionHandle& InHandle, FName Name) const
+			{
+				return ContainsCacheComponent(ActionComponents, InHandle, Name);
+			}
+
+			UAudioComponent* FindComponent(FName Name) const
+			{
+				const FInstanceAudioComponentsPtr* InstanceComponents = GlobalComponents.Find(Name);
+				if (InstanceComponents && InstanceComponents->IsValid())
+				{
+					return (*InstanceComponents)->ActiveComponent.Get();
+				}
+
+				return nullptr;
+			}
+
+			UAudioComponent* FindComponent(const FCollectionHandle& InHandle, FName Name) const
+			{
+				return FindCacheComponent(CollectionComponents, InHandle, Name);
+			}
+
+			UAudioComponent* FindComponent(const FEventHandle& InHandle, FName Name) const
+			{
+				return FindCacheComponent(EventComponents, InHandle, Name);
+			}
+
+			UAudioComponent* FindComponent(const FActionHandle& InHandle, FName Name) const
+			{
+				return FindCacheComponent(ActionComponents, InHandle, Name);
+			}
+
+			void OnCollectionUnregistered(const FCollectionHandle& InHandle)
+			{
+				if (FAudioComponentsCache* Cache = CollectionComponents.Find(InHandle))
+				{
+					for (TPair<FName, FInstanceAudioComponentsPtr>& Pair : *Cache)
+					{
+						FInstanceAudioComponentsPtr& InstanceComponents = Pair.Value;
+						check(InstanceComponents.IsValid());
+						if (UAudioComponent* Component = InstanceComponents->ActiveComponent.Get())
+						{
+							Component->Stop();
+						}
+
+						for (const TObjectPtr<UAudioComponent>& ReleasedComponent : InstanceComponents->ReleasedComponents)
+						{
+							if (UAudioComponent* Component = ReleasedComponent.Get())
+							{
+								Component->Stop();
+							}
+						}
+					}
+
+					CollectionComponents.Remove(InHandle);
+
+	#if SUBSONIC_ASSERT_ON_STALE_COLLECTION_EXECUTORS
+					for (TPair<uint32, FCollectionHandle>& Pair : CollectionExecutors)
+					{
+						checkf(Pair.Value.CollectionId != InHandle.CollectionId, TEXT("Executor for Collection '%s' still exists beyond lifetime of collection."), *InHandle.ToString());
+					}
+	#endif // SUBSONIC_ASSERT_ON_STALE_COLLECTION_EXECUTORS
+				}
+			}
+
+			void OnExecutorRegistered(const FSubsonicExecutor& InExecutor)
+			{
+				CollectionExecutors.Add(InExecutor.GetId(), InExecutor.GetCollectionHandle());
+			}
+
+			void OnExecutorUnregistered(const FSubsonicExecutor& InExecutor)
+			{
+				if (const FCollectionHandle* CollectionHandle = CollectionExecutors.Find(InExecutor.GetId()))
+				{
+					if (const FAudioComponentsCache* ScopeComponentMap = CollectionComponents.Find(*CollectionHandle))
+					{
+						for (const TPair<FName, FInstanceAudioComponentsPtr>& Pair : *ScopeComponentMap)
+						{
+							if (FInstanceAudioComponents* Components = Pair.Value.Get())
+							{
+								Components->StopAll();
+							}
+						}
+					}
+					CollectionExecutors.Remove(InExecutor.GetId());
+					CollectionComponents.Remove(*CollectionHandle);
+				}
+			}
+
+			void RemoveComponent(FName Name)
+			{
+				FInstanceAudioComponentsPtr* InstanceComponents = GlobalComponents.Find(Name);
+				if (InstanceComponents && InstanceComponents->IsValid())
+				{
+					(*InstanceComponents)->StopAll();
+					GlobalComponents.Remove(Name);
+				}
+			}
+
+			void RemoveComponent(const FCollectionHandle& InHandle, FName Name)
+			{
+				return RemoveCacheComponent(CollectionComponents, InHandle, Name);
+			}
+
+			void RemoveComponent(const FEventHandle& InHandle, FName Name)
+			{
+				return RemoveCacheComponent(EventComponents, InHandle, Name);
+			}
+
+			void RemoveComponent(const UE::Subsonic::FActionHandle& InHandle, FName Name)
+			{
+				return RemoveCacheComponent(ActionComponents, InHandle, Name);
+			}
+		};
+	} // namespace AudioComponentEventSubscriberPrivate
+
+	void USubsonicAudioComponentSubscriber::Initialize(FSubsystemCollectionBase& Collection)
+	{
+		Super::Initialize(Collection);
+		Impl = MakeShared<AudioComponentEventSubscriberPrivate::FSubscriberImpl>();
+	}
+
+	void USubsonicAudioComponentSubscriber::Deinitialize()
+	{
+		Impl->Deinitialize();
+		Super::Deinitialize();
+	}
+
+	UAudioComponent* USubsonicAudioComponentSubscriber::AddComponent(FName Name, bool bReleaseExisting)
+	{
+		return Impl->AddComponent(*this, GetAudioDeviceId(), Name, bReleaseExisting);
+	}
+
+	UAudioComponent* USubsonicAudioComponentSubscriber::AddComponent(const FCollectionHandle& InHandle, FName Name, bool bReleaseExisting)
+	{
+		return Impl->AddComponent(*this, GetAudioDeviceId(), InHandle, Name, bReleaseExisting);
+	}
+
+	UAudioComponent* USubsonicAudioComponentSubscriber::AddComponent(const FEventHandle& InHandle, FName Name, bool bReleaseExisting)
+	{
+		return Impl->AddComponent(*this, GetAudioDeviceId(), InHandle, Name, bReleaseExisting);
+	}
+
+	UAudioComponent* USubsonicAudioComponentSubscriber::AddComponent(const UE::Subsonic::FActionHandle& InHandle, FName Name, bool bReleaseExisting)
+	{
+		return Impl->AddComponent(*this, GetAudioDeviceId(), InHandle, Name, bReleaseExisting);
+	}
+
+	bool USubsonicAudioComponentSubscriber::ContainsComponent(FName Name) const
+	{
+		return Impl->ContainsComponent(Name);
+	}
+
+	bool USubsonicAudioComponentSubscriber::ContainsComponent(const FCollectionHandle& InHandle, FName Name) const
+	{
+		return Impl->ContainsComponent(InHandle, Name);
+	}
+
+	bool USubsonicAudioComponentSubscriber::ContainsComponent(const FEventHandle& InHandle, FName Name) const
+	{
+		return Impl->ContainsComponent(InHandle, Name);
+	}
+
+	bool USubsonicAudioComponentSubscriber::ContainsComponent(const UE::Subsonic::FActionHandle& InHandle, FName Name) const
+	{
+		return Impl->ContainsComponent(InHandle, Name);
+	}
+
+	UAudioComponent* USubsonicAudioComponentSubscriber::FindComponent(FName Name)
+	{
+		return Impl->FindComponent(Name);
+	}
+
+	UAudioComponent* USubsonicAudioComponentSubscriber::FindComponent(const FCollectionHandle& InHandle, FName Name)
+	{
+		return Impl->FindComponent(InHandle, Name);
+	}
+
+	UAudioComponent* USubsonicAudioComponentSubscriber::FindComponent(const FEventHandle& InHandle, FName Name)
+	{
+		return Impl->FindComponent(InHandle, Name);
+	}
+
+	UAudioComponent* USubsonicAudioComponentSubscriber::FindComponent(const UE::Subsonic::FActionHandle& InHandle, FName Name)
+	{
+		return Impl->FindComponent(InHandle, Name);
+	}
+
+	void USubsonicAudioComponentSubscriber::RemoveComponent(FName Name)
+	{
+		return Impl->RemoveComponent(Name);
+	}
+
+	void USubsonicAudioComponentSubscriber::RemoveComponent(const FCollectionHandle& InHandle, FName Name)
+	{
+		return Impl->RemoveComponent(InHandle, Name);
+	}
+
+	void USubsonicAudioComponentSubscriber::RemoveComponent(const FEventHandle& InHandle, FName Name)
+	{
+		return Impl->RemoveComponent(InHandle, Name);
+	}
+
+	void USubsonicAudioComponentSubscriber::OnCollectionUnregistered(const FCollectionHandle& InHandle)
+	{
+		Impl->OnCollectionUnregistered(InHandle);
+	}
+
+	void USubsonicAudioComponentSubscriber::OnExecutorRegistered(const FSubsonicExecutor& InExecutor)
+	{
+		Impl->OnExecutorRegistered(InExecutor);
+	}
+
+	void USubsonicAudioComponentSubscriber::OnExecutorUnregistered(const FSubsonicExecutor& InExecutor)
+	{
+		Impl->OnExecutorUnregistered(InExecutor);
+	}
+} // namespace UE::Subsonic
+#undef SUBSONIC_ASSERT_ON_STALE_COLLECTION_EXECUTORS
diff --git a/Engine/Plugins/Experimental/Subsonic/Source/Subsonic/Private/CoreEventSubscribers/AudioComponentEventSubscriber.h b/Engine/Plugins/Experimental/Subsonic/Source/Subsonic/Private/CoreEventSubscribers/AudioComponentEventSubscriber.h
new file mode 100644
--- /dev/null
+++ b/Engine/Plugins/Experimental/Subsonic/Source/Subsonic/Private/CoreEventSubscribers/AudioComponentEventSubscriber.h
@@ -0,0 +1,71 @@
+// Copyright Epic Games, Inc. All Rights Reserved.
+#pragma once
+
+#include "Subsystems/EngineSubsystem.h"
+
+#include "Components/AudioComponent.h"
+#include "Containers/Map.h"
+#include "CoreEventSubscribers/SubsonicEventSubscriberBase.h"
+#include "SubsonicEventSystemInterface.h"
+#include "SubsonicHandles.h"
+#include "Subsystems/AudioEngineSubsystem.h"
+#include "Templates/SharedPointer.h"
+
+#include "AudioComponentEventSubscriber.generated.h"
+
+#define UE_API SUBSONIC_API
+
+
+// Forward Declarations
+namespace UE::Subsonic
+{
+	struct FSubsonicExecutor;
+
+	namespace AudioComponentEventSubscriberPrivate
+	{
+		class FSubscriberImpl;
+	} // namespace AudioComponentEventSubscriberPrivate
+
+	/**
+	 * Subsonic Subscriber managing AudioComponent instances, events, and related data.
+	 */
+	UCLASS(MinimalAPI, BlueprintType)
+	class USubsonicAudioComponentSubscriber final : public USubsonicEventSubscriberBase
+	{
+		GENERATED_BODY()
+
+	public:
+		//~ Begin USubsonicEventSubscriberBase interface
+		UE_API virtual void Initialize(FSubsystemCollectionBase& Collection) override;
+		UE_API virtual void Deinitialize() override;
+		//~ End USubsonicEventSubscriberBase interface
+
+		UE_INTERNAL virtual void OnCollectionUnregistered(const FCollectionHandle& InEvent) override;
+		UE_INTERNAL virtual void OnExecutorRegistered(const FSubsonicExecutor& InExecutor) override;
+		UE_INTERNAL virtual void OnExecutorUnregistered(const FSubsonicExecutor& InExecutor) override;
+
+		UE_INTERNAL UAudioComponent* AddComponent(FName Name, bool bReleaseExisting = true);
+		UE_INTERNAL UAudioComponent* AddComponent(const FCollectionHandle& InHandle, FName Name, bool bReleaseExisting = true);
+		UE_INTERNAL UAudioComponent* AddComponent(const FEventHandle& InHandle, FName Name, bool bReleaseExisting = true);
+		UE_INTERNAL UAudioComponent* AddComponent(const FActionHandle& InHandle, FName Name, bool bReleaseExisting = true);
+
+		UE_INTERNAL bool ContainsComponent(FName Name) const;
+		UE_INTERNAL bool ContainsComponent(const FCollectionHandle& InHandle, FName Name) const;
+		UE_INTERNAL bool ContainsComponent(const FEventHandle& InHandle, FName Name) const;
+		UE_INTERNAL bool ContainsComponent(const FActionHandle& InHandle, FName Name) const;
+
+		UE_INTERNAL UAudioComponent* FindComponent(FName Name);
+		UE_INTERNAL UAudioComponent* FindComponent(const FCollectionHandle& InHandle, FName Name);
+		UE_INTERNAL UAudioComponent* FindComponent(const FEventHandle& InHandle, FName Name);
+		UE_INTERNAL UAudioComponent* FindComponent(const FActionHandle& InHandle, FName Name);
+
+		UE_INTERNAL void RemoveComponent(FName Name);
+		UE_INTERNAL void RemoveComponent(const FCollectionHandle& InHandle, FName Name);
+		UE_INTERNAL void RemoveComponent(const FEventHandle& InHandle, FName Name);
+
+	private:
+		TSharedPtr<AudioComponentEventSubscriberPrivate::FSubscriberImpl> Impl;
+	};
+} // namespace UE::Subsonic
+
+	#undef UE_API
diff --git a/Engine/Plugins/Experimental/Subsonic/Source/Subsonic/Private/CoreEventSubscribers/SubsonicEventSubscriberBase.cpp b/Engine/Plugins/Experimental/Subsonic/Source/Subsonic/Private/CoreEventSubscribers/SubsonicEventSubscriberBase.cpp
new file mode 100644
--- /dev/null
+++ b/Engine/Plugins/Experimental/Subsonic/Source/Subsonic/Private/CoreEventSubscribers/SubsonicEventSubscriberBase.cpp
@@ -0,0 +1,34 @@
+// Copyright Epic Games, Inc. All Rights Reserved.
+
+#include "CoreEventSubscribers/SubsonicEventSubscriberBase.h"
+#include "Engine/Engine.h"
+#include "SubsonicSubsystem.h"
+
+namespace UE::Subsonic
+{
+	void USubsonicEventSubscriberBase::Initialize(FSubsystemCollectionBase& Collection)
+	{
+		if (USubsonicSubsystem* SubsonicSubsystem = GetSubsonicSubsystem())
+		{
+			SubsonicSubsystem->OnEventSubscriberRegistered(this);
+		}
+	}
+
+	void USubsonicEventSubscriberBase::Deinitialize()
+	{
+		if (USubsonicSubsystem* SubsonicSubsystem = GetSubsonicSubsystem())
+		{
+			SubsonicSubsystem->OnEventSubscriberUnregistered(this);
+		}
+	}
+
+	USubsonicSubsystem* USubsonicEventSubscriberBase::GetSubsonicSubsystem() const
+	{
+		if (GEngine)
+		{
+			return GEngine->GetEngineSubsystem<USubsonicSubsystem>();
+		}
+
+		return nullptr;
+	}
+} // namespace UE::Subsonic
\ No newline at end of file
diff --git a/Engine/Plugins/Experimental/Subsonic/Source/Subsonic/Private/CoreEventSystems/SubsonicAudioComponentSubsystem.cpp b/Engine/Plugins/Experimental/Subsonic/Source/Subsonic/Private/CoreEventSystems/SubsonicAudioComponentSubsystem.cpp
deleted file mode 100644
diff --git a/Engine/Plugins/Experimental/Subsonic/Source/Subsonic/Private/CoreEventSystems/SubsonicAudioComponentSubsystem.h b/Engine/Plugins/Experimental/Subsonic/Source/Subsonic/Private/CoreEventSystems/SubsonicAudioComponentSubsystem.h
deleted file mode 100644
diff --git a/Engine/Plugins/Experimental/Subsonic/Source/Subsonic/Private/SubsonicEventCollection.cpp b/Engine/Plugins/Experimental/Subsonic/Source/Subsonic/Private/SubsonicEventCollection.cpp
--- a/Engine/Plugins/Experimental/Subsonic/Source/Subsonic/Private/SubsonicEventCollection.cpp
+++ b/Engine/Plugins/Experimental/Subsonic/Source/Subsonic/Private/SubsonicEventCollection.cpp
@@ -14,148 +14,171 @@ namespace UE::Subsonic
 {
 	UE_DEFINE_GAMEPLAY_TAG(TAG_SubsonicEventCorePlay, "Subsonic.Event.Play");
 	UE_DEFINE_GAMEPLAY_TAG(TAG_SubsonicEventCoreStop, "Subsonic.Event.Stop");
-} // namespace UE::Subsonic
 
+	namespace CollectionPrivate
+	{
+		constexpr uint32 InvalidId = static_cast<uint32>(INDEX_NONE);
+	} // namespace CollectionPrivate
 
-void FSubsonicEvent::Execute(const UE::Subsonic::FEventHandle& InHandle) const
-{
-	using namespace UE::Subsonic;
 
-	FActionHandle ActionHandle { .Event = InHandle };
-	for (int32 Index = 0; Index < Actions.Num(); ++Index)
+	void FSubsonicEvent::Execute(const FSubsonicExecutor& InExecutor, const FEventHandle& InHandle) const
 	{
-		ActionHandle.Index = Index;
+		checkf(InExecutor.GetCollectionHandle().CollectionId == InHandle.Collection.CollectionId, TEXT("Invalid execution context: executor CollectionId must match that of the given event handle"));
 
-		TConstStructView<FSubsonicEventActionBase> Action = Actions[Index];
-		if (const FSubsonicEventActionBase* ActionPtr = Action.GetPtr())
-		{
-			ActionPtr->Execute(ActionHandle);
-		}
-		else
+		FActionHandle ActionHandle { .Event = InHandle };
+		for (int32 Index = 0; Index < Actions.Num(); ++Index)
 		{
-			UE_LOG(LogSubsonic, Warning, TEXT("Action not specified: Skipping action at '%s'"), *ActionHandle.ToString());
+			ActionHandle.Index = Index;
+
+			TConstStructView<FSubsonicEventActionBase> Action = Actions[Index];
+			if (const FSubsonicEventActionBase* ActionPtr = Action.GetPtr())
+			{
+				ActionPtr->Execute(InExecutor, ActionHandle);
+			}
+			else
+			{
+				UE_LOG(LogSubsonic, Warning, TEXT("Action not specified: Skipping action at '%s'"), *ActionHandle.ToString());
+			}
 		}
 	}
-}
 
-UE::Subsonic::FCollectionHandle USubsonicEventCollection::GetHandle(const FSubsonicEventCollectionExecutor& Executor) const
-{
-	using namespace UE::Subsonic;
+	FSubsonicEventCollectionDefinition::FSubsonicEventCollectionDefinition(FSubsonicEventCollectionDefinition&& InOther)
+	{
+		Unregister();
+		CollectionId = InOther.CollectionId;
+		InOther.CollectionId = GetInvalidId();
 
-	FName CollectionName;
+		Events = MoveTemp(InOther.Events);
+	}
 
-#if WITH_EDITORONLY_DATA
-	if (const UPackage* CollectionPackage = GetPackage())
+	FSubsonicEventCollectionDefinition& FSubsonicEventCollectionDefinition::operator =(FSubsonicEventCollectionDefinition&& InOther)
 	{
-		const FTopLevelAssetPath AssetPath(CollectionPackage);
-		CollectionName = AssetPath.GetPackageName();
+		Unregister();
+		CollectionId = InOther.CollectionId;
+		InOther.CollectionId = GetInvalidId();
+
+		Events = MoveTemp(InOther.Events);
+		return *this;
 	}
-#endif // WITH_EDITORONLY_DATA
 
-	return FCollectionHandle
+	void FSubsonicEventCollectionDefinition::AssignId()
 	{
-#if WITH_EDITORONLY_DATA
-		.CollectionName = CollectionName,
-#endif // WITH_EDITORONLY_DATA
-		.CollectionId = GetUniqueID(),
-		.DeviceId = Executor.GetDeviceId(),
-		.ExecutorId = Executor.GetExecutorId()
-	};
-}
+		static uint32 MaxId = CollectionPrivate::InvalidId;
+		CollectionId = ++MaxId;
+	}
 
-void USubsonicEventCollection::PostLoad()
-{
-	Super::PostLoad();
+	FSubsonicEventCollectionDefinition FSubsonicEventCollectionDefinition::Create(FName Name, TMap<FGameplayTag, FSubsonicEvent> InEvents)
+	{
+		FSubsonicEventCollectionDefinition NewDefinition;
+		NewDefinition.AssignId();
 
-	using namespace UE::Subsonic;
+		const FCollectionHandle NewHandle
+		{
+	#if WITH_EDITORONLY_DATA
+				.CollectionName = Name,
+	#endif // WITH_EDITORONLY_DATA
+				.CollectionId = NewDefinition.GetCollectionId(),
+		};
+
+		NewDefinition.Events = MoveTemp(InEvents);
+		NewDefinition.Register(NewHandle);
+		return MoveTemp(NewDefinition);
+	}
 
-	if (GEngine)
+	FSubsonicEventCollectionDefinition::~FSubsonicEventCollectionDefinition()
 	{
-		if (USubsonicSubsystem* Subsystem = GEngine->GetEngineSubsystem<USubsonicSubsystem>())
-		{
-			FName CollectionName;
+		Unregister();
+		CollectionId = GetInvalidId();
+	}
 
-#if WITH_EDITORONLY_DATA
-			if (const UPackage* CollectionPackage = GetPackage())
-			{
-				const FTopLevelAssetPath AssetPath(CollectionPackage);
-				CollectionName = AssetPath.GetPackageName();
-			}
-#endif // WITH_EDITORONLY_DATA
+	bool FSubsonicEventCollectionDefinition::IsValid() const
+	{
+		return CollectionId != INDEX_NONE;
+	}
 
-			Subsystem->RegisterCollection(FCollectionHandle
-			{
-#if WITH_EDITORONLY_DATA
-				.CollectionName = CollectionName,
-#endif // WITH_EDITORONLY_DATA
-				.CollectionId = GetUniqueID(),
-			});
-		}
+	uint32 FSubsonicEventCollectionDefinition::GetCollectionId() const
+	{
+		return CollectionId;
 	}
-}
 
-void USubsonicEventCollection::BeginDestroy()
-{
-	using namespace UE::Subsonic;
+	const TMap<FGameplayTag, FSubsonicEvent>& FSubsonicEventCollectionDefinition::GetEvents() const
+	{
+		return Events;
+	}
 
-	if (GEngine)
+	uint32 FSubsonicEventCollectionDefinition::GetInvalidId()
 	{
-		if (USubsonicSubsystem* Subsystem = GEngine->GetEngineSubsystem<USubsonicSubsystem>())
-		{
-			FName CollectionName;
+		return UE::Subsonic::CollectionPrivate::InvalidId;
+	}
 
-#if WITH_EDITORONLY_DATA
-			if (const UPackage* CollectionPackage = GetPackage())
+	void FSubsonicEventCollectionDefinition::Register(const FCollectionHandle& InHandle)
+	{
+		if (GEngine && IsValid())
+		{
+			if (USubsonicSubsystem* Subsystem = GEngine->GetEngineSubsystem<USubsonicSubsystem>())
 			{
-				const FTopLevelAssetPath AssetPath(CollectionPackage);
-				CollectionName = AssetPath.GetPackageName();
+				Subsystem->OnCollectionRegistered(InHandle);
 			}
-#endif // WITH_EDITORONLY_DATA
+		}
+	}
 
-			Subsystem->UnregisterCollection(FCollectionHandle
+	void FSubsonicEventCollectionDefinition::Unregister()
+	{
+		if (GEngine && IsValid())
+		{
+			if (USubsonicSubsystem* Subsystem = GEngine->GetEngineSubsystem<USubsonicSubsystem>())
 			{
-#if WITH_EDITORONLY_DATA
-				.CollectionName = CollectionName,
-#endif // WITH_EDITORONLY_DATA
-				.CollectionId = GetUniqueID(),
-			});
+				Subsystem->OnCollectionUnregistered(FCollectionHandle { .CollectionId = CollectionId });
+			}
 		}
+
+		CollectionId = CollectionPrivate::InvalidId;
 	}
 
-	Super::BeginDestroy();
-}
+	void USubsonicEventCollection::PostInitProperties()
+	{
+		Super::PostInitProperties();
+		CollectionDefinition.AssignId();
+	}
 
-uint32 FSubsonicEventCollectionExecutor::MaxExecutorId = INDEX_NONE;
+	void USubsonicEventCollection::PostLoad()
+	{
+		checkf(CollectionDefinition.IsValid(), TEXT("Definition should have an assigned Id by this point (via PostInitProperties)"));
 
-FSubsonicEventCollectionExecutor::FSubsonicEventCollectionExecutor(const USubsonicEventCollection* InCollection, Audio::FDeviceId InDeviceId)
-	: Collection(InCollection)
-	, DeviceId(InDeviceId)
-	, ExecutorId(++MaxExecutorId)
-{
-}
+		Super::PostLoad();
+		CollectionDefinition.Register(GetHandle());
+	}
 
-Audio::FDeviceId FSubsonicEventCollectionExecutor::GetDeviceId() const
-{
-	return DeviceId;
-}
+	void USubsonicEventCollection::BeginDestroy()
+	{
+		CollectionDefinition.Unregister();
+		Super::BeginDestroy();
+	}
 
-const USubsonicEventCollection& FSubsonicEventCollectionExecutor::GetCollectionChecked() const
-{
-	check(Collection.Get());
-	return *Collection.Get();
-}
+	void USubsonicEventCollectionExecutor::BeginDestroy()
+	{
+		Executor.Unregister();
+		Super::BeginDestroy();
+	}
 
-const USubsonicEventCollection* FSubsonicEventCollectionExecutor::GetCollection() const
-{
-	return Collection.Get();
-}
+	const USubsonicEventCollection& USubsonicEventCollectionExecutor::GetCollectionChecked() const
+	{
+		check(Collection.Get());
+		return *Collection.Get();
+	}
 
-uint32 FSubsonicEventCollectionExecutor::GetExecutorId() const
-{
-	return ExecutorId;
-}
+	const USubsonicEventCollection* USubsonicEventCollectionExecutor::GetCollection() const
+	{
+		return Collection.Get();
+	}
 
-bool FSubsonicEventCollectionExecutor::IsValid() const
-{
-	return ExecutorId != INDEX_NONE && Collection.Get() != nullptr;
-}
+	const FSubsonicExecutor& USubsonicEventCollectionExecutor::GetExecutor() const
+	{
+		return Executor;
+	}
+
+	bool USubsonicEventCollectionExecutor::IsValid() const
+	{
+		return Executor.IsValid() && GetCollection() != nullptr;
+	}
+} // namespace UE::Subsonic
diff --git a/Engine/Plugins/Experimental/Subsonic/Source/Subsonic/Private/SubsonicEventInstanceImpl.cpp b/Engine/Plugins/Experimental/Subsonic/Source/Subsonic/Private/SubsonicEventInstanceImpl.cpp
--- a/Engine/Plugins/Experimental/Subsonic/Source/Subsonic/Private/SubsonicEventInstanceImpl.cpp
+++ b/Engine/Plugins/Experimental/Subsonic/Source/Subsonic/Private/SubsonicEventInstanceImpl.cpp
@@ -6,22 +6,26 @@
 #include "Engine/World.h"
 #include "Kismet/GameplayStatics.h"
 
-FSubsonicEventInstanceImpl::FSubsonicEventInstanceImpl()
-{
-}
 
-void FSubsonicEventInstanceImpl::PlaySound(const UObject* WorldContextObject, USoundBase* Sound)
+namespace UE::Subsonic
 {
-	AudioComponents.Add(UGameplayStatics::SpawnSound2D(WorldContextObject, Sound));
-}
+	FSubsonicEventInstanceImpl::FSubsonicEventInstanceImpl()
+	{
+	}
 
-void FSubsonicEventInstanceImpl::Stop()
-{
-	for (TWeakObjectPtr<UAudioComponent> AudioComponent : AudioComponents)
+	void FSubsonicEventInstanceImpl::PlaySound(const UObject* WorldContextObject, USoundBase* Sound)
+	{
+		AudioComponents.Add(UGameplayStatics::SpawnSound2D(WorldContextObject, Sound));
+	}
+
+	void FSubsonicEventInstanceImpl::Stop()
 	{
-		if (AudioComponent.IsValid())
+		for (TWeakObjectPtr<UAudioComponent> AudioComponent : AudioComponents)
 		{
-			AudioComponent->Stop();
+			if (AudioComponent.IsValid())
+			{
+				AudioComponent->Stop();
+			}
 		}
 	}
-}
\ No newline at end of file
+} // namespace UE::Subsonic
\ No newline at end of file
diff --git a/Engine/Plugins/Experimental/Subsonic/Source/Subsonic/Private/SubsonicEventInstanceImpl.h b/Engine/Plugins/Experimental/Subsonic/Source/Subsonic/Private/SubsonicEventInstanceImpl.h
--- a/Engine/Plugins/Experimental/Subsonic/Source/Subsonic/Private/SubsonicEventInstanceImpl.h
+++ b/Engine/Plugins/Experimental/Subsonic/Source/Subsonic/Private/SubsonicEventInstanceImpl.h
@@ -7,15 +7,19 @@
 class UAudioComponent;
 class USoundBase;
 
-class FSubsonicEventInstanceImpl
+
+namespace UE::Subsonic
 {
-public:
-	FSubsonicEventInstanceImpl();
+	class FSubsonicEventInstanceImpl
+	{
+	public:
+		FSubsonicEventInstanceImpl();
 
-	void PlaySound(const UObject* WorldContextObject, USoundBase* Sound);
+		void PlaySound(const UObject* WorldContextObject, USoundBase* Sound);
 
-	void Stop();
+		void Stop();
 
-private:
-	TArray<TWeakObjectPtr<UAudioComponent>> AudioComponents;
-};
+	private:
+		TArray<TWeakObjectPtr<UAudioComponent>> AudioComponents;
+	};
+} // namespace UE::Subsonic
diff --git a/Engine/Plugins/Experimental/Subsonic/Source/Subsonic/Private/SubsonicEventSystemInterface.h b/Engine/Plugins/Experimental/Subsonic/Source/Subsonic/Private/SubsonicEventSystemInterface.h
--- a/Engine/Plugins/Experimental/Subsonic/Source/Subsonic/Private/SubsonicEventSystemInterface.h
+++ b/Engine/Plugins/Experimental/Subsonic/Source/Subsonic/Private/SubsonicEventSystemInterface.h
@@ -2,25 +2,31 @@
 #pragma once
 
 #include "UObject/Interface.h"
-#include "SubsonicHandles.h"
+#include "SubsonicExecutor.h"
 
 #include "SubsonicEventSystemInterface.generated.h"
 
-// Interface for all subscribers that respond to Subsonic events.
-UINTERFACE(MinimalAPI, BlueprintType, meta = (DisplayName = "Subsonic Event System Interface", CannotImplementInterfaceInBlueprint))
-class USubsonicEventSystemInterface : public UInterface
+namespace UE::Subsonic
 {
-	GENERATED_BODY()
-};
+	// Interface for all subscribers that respond to Subsonic events.
+	UINTERFACE(MinimalAPI, BlueprintType, meta = (DisplayName = "Subsonic Event Subscriber Interface", CannotImplementInterfaceInBlueprint))
+	class USubsonicEventSubscriberInterface : public UInterface
+	{
+		GENERATED_BODY()
+	};
 
-class ISubsonicEventSystemInterface : public IInterface
-{
-	GENERATED_BODY()
+	class ISubsonicEventSubscriberInterface : public IInterface
+	{
+		GENERATED_BODY()
+
+	public:
+		virtual void OnCollectionRegistered(const FCollectionHandle& InCollection) { }
+		virtual void OnCollectionUnregistered(const FCollectionHandle& InCollection) { }
 
-public:
-	virtual void OnEventPreExecute(const UE::Subsonic::FEventHandle& InHandle) = 0;
-	virtual void OnEventPostExecute(const UE::Subsonic::FEventHandle& InHandle) = 0;
+		virtual void OnEventPreExecute(const FSubsonicExecutor& InExecutor, const FEventHandle& InHandle) { }
+		virtual void OnEventPostExecute(const FSubsonicExecutor& InExecutor, const FEventHandle& InHandle) { }
 
-	virtual void OnCollectionRegistered(const UE::Subsonic::FCollectionHandle& InEvent) = 0;
-	virtual void OnCollectionUnregistered(const UE::Subsonic::FCollectionHandle& InEvent) = 0;
-};
+		virtual void OnExecutorRegistered(const FSubsonicExecutor& InExecutor) { }
+		virtual void OnExecutorUnregistered(const FSubsonicExecutor& InExecutor) { }
+	};
+} // namespace UE::Subsonic
\ No newline at end of file
diff --git a/Engine/Plugins/Experimental/Subsonic/Source/Subsonic/Private/SubsonicExecutor.cpp b/Engine/Plugins/Experimental/Subsonic/Source/Subsonic/Private/SubsonicExecutor.cpp
new file mode 100644
--- /dev/null
+++ b/Engine/Plugins/Experimental/Subsonic/Source/Subsonic/Private/SubsonicExecutor.cpp
@@ -0,0 +1,153 @@
+// Copyright Epic Games, Inc. All Rights Reserved.
+
+#include "SubsonicExecutor.h"
+
+#include "Engine/Engine.h"
+#include "SubsonicEventCollection.h"
+#include "SubsonicHandles.h"
+#include "SubsonicSubsystem.h"
+
+
+namespace UE::Subsonic
+{
+	namespace ExecutorPrivate
+	{
+		constexpr uint32 InvalidId = static_cast<uint32>(INDEX_NONE);
+		Audio::FDeviceId InvalidDeviceId = static_cast<Audio::FDeviceId>(INDEX_NONE);
+	} // namespace ExecutorPrivate
+
+	FSubsonicExecutor::FSubsonicExecutor()
+		: DeviceId(ExecutorPrivate::InvalidDeviceId)
+	{
+	}
+
+	FSubsonicExecutor::FSubsonicExecutor(FSubsonicExecutor&& InOther)
+	{
+		Unregister();
+
+		CollectionHandle = InOther.CollectionHandle;
+		DeviceId = InOther.DeviceId;
+		ExecutorId = InOther.ExecutorId;
+
+		InOther.CollectionHandle = { };
+		InOther.DeviceId = ExecutorPrivate::InvalidDeviceId;
+		InOther.ExecutorId = ExecutorPrivate::InvalidId;
+	}
+
+	FSubsonicExecutor::~FSubsonicExecutor()
+	{
+		Unregister();
+	}
+
+	FSubsonicExecutor& FSubsonicExecutor::operator=(FSubsonicExecutor&& InOther)
+	{
+		Unregister();
+
+		CollectionHandle = InOther.CollectionHandle;
+		DeviceId = InOther.DeviceId;
+		ExecutorId = InOther.ExecutorId;
+
+		InOther.CollectionHandle = { };
+		InOther.DeviceId = ExecutorPrivate::InvalidDeviceId;
+		InOther.ExecutorId = ExecutorPrivate::InvalidId;
+
+		return *this;
+	}
+
+	FSubsonicExecutor FSubsonicExecutor::Create(Audio::FDeviceId InDeviceId, FCollectionHandle CollectionHandle)
+	{
+		static uint32 MaxExecutorId = ExecutorPrivate::InvalidId;
+		FSubsonicExecutor NewHandle;
+		NewHandle.CollectionHandle = MoveTemp(CollectionHandle);
+		NewHandle.ExecutorId = ++MaxExecutorId;
+		NewHandle.DeviceId = InDeviceId;
+		NewHandle.Register();
+		return NewHandle;
+	}
+
+	FCollectionHandle FSubsonicExecutor::GetCollectionHandle() const
+	{
+		return CollectionHandle;
+	}
+
+	Audio::FDeviceId FSubsonicExecutor::GetDeviceId() const
+	{
+		return DeviceId;
+	}
+
+	uint32 FSubsonicExecutor::GetId() const
+	{
+		return ExecutorId;
+	}
+
+	void FSubsonicExecutor::ExecuteEvent(const FCollectionHandle& InCollection, FName EventName, const FSubsonicEvent& InEvent) const
+	{
+		const FEventHandle EventHandle
+		{
+			.Collection = InCollection,
+			.EventName = EventName
+		};
+
+		if (GEngine)
+		{
+			if (USubsonicSubsystem* Subsystem = GEngine->GetEngineSubsystem<USubsonicSubsystem>())
+			{
+				Subsystem->OnEventPreExecute(*this, EventHandle);
+			}
+		}
+
+		InEvent.Execute(*this, EventHandle);
+
+		if (GEngine)
+		{
+			if (USubsonicSubsystem* Subsystem = GEngine->GetEngineSubsystem<USubsonicSubsystem>())
+			{
+				Subsystem->OnEventPostExecute(*this, EventHandle);
+			}
+		}
+	}
+
+	uint32 FSubsonicExecutor::GetInvalidId()
+	{
+		return ExecutorPrivate::InvalidId;
+	}
+
+	FSubsonicExecutor FSubsonicExecutor::GetInvalid()
+	{
+		return { };
+	}
+
+	bool FSubsonicExecutor::IsValid() const
+	{
+		return ExecutorId != ExecutorPrivate::InvalidId;
+	}
+
+	void FSubsonicExecutor::Register()
+	{
+		if (GEngine)
+		{
+			if (USubsonicSubsystem* Subsystem = GEngine->GetEngineSubsystem<USubsonicSubsystem>())
+			{
+				Subsystem->OnExecutorRegistered(*this);
+			}
+		}
+	}
+
+	void FSubsonicExecutor::Unregister()
+	{
+		if (GEngine && IsValid())
+		{
+			if (USubsonicSubsystem* Subsystem = GEngine->GetEngineSubsystem<USubsonicSubsystem>())
+			{
+				Subsystem->OnExecutorUnregistered(*this);
+			}
+		}
+
+		ExecutorId = ExecutorPrivate::InvalidId;
+	}
+
+	FString FSubsonicExecutor::ToString() const
+	{
+		return FString::Printf(TEXT("Id: %u, DeviceId: %u"), ExecutorId, DeviceId);
+	}
+} // namespace UE::Subsonic
\ No newline at end of file
diff --git a/Engine/Plugins/Experimental/Subsonic/Source/Subsonic/Private/SubsonicHandles.cpp b/Engine/Plugins/Experimental/Subsonic/Source/Subsonic/Private/SubsonicHandles.cpp
new file mode 100644
--- /dev/null
+++ b/Engine/Plugins/Experimental/Subsonic/Source/Subsonic/Private/SubsonicHandles.cpp
@@ -0,0 +1,36 @@
+// Copyright Epic Games, Inc. All Rights Reserved.
+
+#include "SubsonicHandles.h"
+
+
+namespace UE::Subsonic
+{
+	namespace HandlesPrivate
+	{
+		uint32 InvalidHandleId = static_cast<uint32>(INDEX_NONE);
+	} // namespace HandlesPrivate
+
+	uint32 FCollectionHandle::GetInvalidId()
+	{
+		return HandlesPrivate::InvalidHandleId;
+	}
+
+	FString FCollectionHandle::ToString() const
+	{
+#if WITH_EDITORONLY_DATA
+		return FString::Printf(TEXT("%s (Id: %u)"), *CollectionName.ToString(), CollectionId);
+#else // !WITH_EDITORONLY_DATA
+		return FString::Printf(TEXT("Id: %u"), CollectionId);
+#endif // !WITH_EDITORONLY_DATA
+	}
+
+	FString FEventHandle::ToString() const
+	{
+		return FString::Printf(TEXT("%s/%s"), *Collection.ToString(), *EventName.ToString());
+	}
+
+	FString FActionHandle::ToString() const
+	{
+		return FString::Printf(TEXT("%s[%i]"), *Event.ToString(), Index);
+	}
+} // namespace UE::Subsonic
\ No newline at end of file
diff --git a/Engine/Plugins/Experimental/Subsonic/Source/Subsonic/Private/SubsonicSubsystem.cpp b/Engine/Plugins/Experimental/Subsonic/Source/Subsonic/Private/SubsonicSubsystem.cpp
--- a/Engine/Plugins/Experimental/Subsonic/Source/Subsonic/Private/SubsonicSubsystem.cpp
+++ b/Engine/Plugins/Experimental/Subsonic/Source/Subsonic/Private/SubsonicSubsystem.cpp
@@ -2,14 +2,15 @@
 
 #include "SubsonicSubsystem.h"
 
+#include "Algo/Transform.h"
 #include "AudioDefines.h"
 #include "Engine/GameInstance.h"
 #include "Engine/Engine.h"
 #include "GameplayTagsManager.h"
 #include "Misc/CoreDefines.h"
-#include "SubsonicEventInstanceImpl.h"
-#include "SubsonicTempDataAsset.h"
+#include "Misc/Optional.h"
 #include "SubsonicEventCollection.h"
+#include "SubsonicEventInstanceImpl.h"
 #include "SubsonicHandles.h"
 #include "UObject/GCObject.h"
 
@@ -18,7 +19,7 @@ namespace UE::Subsonic
 {
 	namespace SubsystemPrivate
 	{
-		Audio::FDeviceId GetDeviceIdFrom(const UObject* WorldContextObject)
+		Audio::FDeviceId GetDeviceIdFromWorldObject(const UObject* WorldContextObject)
 		{
 			if (GEngine)
 			{
@@ -30,185 +31,211 @@ namespace UE::Subsonic
 
 			return INDEX_NONE;
 		};
-	} // namespace SubsystemPrivate
-} // namespace UE::Subsonic
 
-class FSubsonicSubsystemImpl : FGCObject
-{
-public:
-	FSubsonicEventInstance PlayEvent(const UObject* WorldContextObject, const FName EventName);
-	void StopEvent(FSubsonicEventInstance EventInstance);
-	void RegisterHandlers(USubsonicTempDataAsset* Handlers);
+		TOptional<Audio::FDeviceId> GetDeviceIdFromSubscriber(const TScriptInterface<ISubsonicEventSubscriberInterface>& Subscriber)
+		{
+			UObject* SubscriberObject = Subscriber.GetObject();
+			if (!SubscriberObject)
+			{
+				return { };
+			}
 
-private:
-	virtual void AddReferencedObjects(FReferenceCollector& Collector) override
-	{
-		Collector.AddStableReferenceArray(&RegisteredHandlers);
-	}
-	virtual FString GetReferencerName() const override
-	{
-		return "FSubsonicSubsystemImpl";
-	}
+			if (const UAudioEngineSubsystem* Subsystem = Cast<UAudioEngineSubsystem>(Subscriber.GetObject()))
+			{
+				return Subsystem->GetAudioDeviceId();
+			}
 
-	uint32 NextHandle = 1;
+			return GetDeviceIdFromWorldObject(Subscriber.GetObject());
 
-	TMap<FSubsonicEventInstance, FSubsonicEventInstanceImpl> ActiveInstances;
+		}
 
-	TArray<TObjectPtr<USubsonicTempDataAsset>> RegisteredHandlers;
-};
+		class FSubsonicSubsystemImpl : public FGCObject
+		{
+		public:
+			void OnCollectionRegistered(const FCollectionHandle& InHandle)
+			{
+				for (TScriptInterface<ISubsonicEventSubscriberInterface>& EventSubscriber : EventSubscribers)
+				{
+					if (EventSubscriber.GetObject() != nullptr)
+					{
+						EventSubscriber->OnCollectionRegistered(InHandle);
+					}
+				}
+			}
 
-FSubsonicEventInstance FSubsonicSubsystemImpl::PlayEvent(const UObject* WorldContextObject, const FName EventName)
-{
-	FSubsonicEventInstance NewHandle;
-	NewHandle.RawHandle = NextHandle++;
+			void OnCollectionUnregistered(const FCollectionHandle& InHandle)
+			{
+				for (TScriptInterface<ISubsonicEventSubscriberInterface>& EventSubscriber : EventSubscribers)
+				{
+					if (EventSubscriber.GetObject() != nullptr)
+					{
+						EventSubscriber->OnCollectionUnregistered(InHandle);
+					}
+				}
+			}
 
-	FSubsonicEventInstanceImpl& NewInstance = ActiveInstances.Emplace(NewHandle, {});
-	for (USubsonicTempDataAsset* Handlers : RegisteredHandlers)
-	{
-		TObjectPtr<USoundBase>* MappedSound = Handlers->EventHandlers.Find(EventName);
-		if (MappedSound && *MappedSound)
-		{
-			NewInstance.PlaySound(WorldContextObject, *MappedSound);
-		}
-	}
+			void OnEventPreExecute(const FSubsonicExecutor& Executor, const FEventHandle& InEventHandle)
+			{
+				for (const TScriptInterface<ISubsonicEventSubscriberInterface>& EventSubscriber : EventSubscribers)
+				{
+					const TOptional<Audio::FDeviceId> SystemDeviceId = SubsystemPrivate::GetDeviceIdFromSubscriber(EventSubscriber);
+					if (SystemDeviceId == Executor.GetDeviceId())
+					{
+						EventSubscriber->OnEventPreExecute(Executor, InEventHandle);
+					}
+				}
+			}
 
-	return NewHandle;
-}
+			void OnEventPostExecute(const FSubsonicExecutor& Executor, const FEventHandle& InEventHandle)
+			{
+				for (const TScriptInterface<ISubsonicEventSubscriberInterface>& EventSubscriber : EventSubscribers)
+				{
+					const TOptional<Audio::FDeviceId> SystemDeviceId = SubsystemPrivate::GetDeviceIdFromSubscriber(EventSubscriber);
+					if (SystemDeviceId == Executor.GetDeviceId())
+					{
+						if (EventSubscriber.GetObject())
+						{
+							EventSubscriber->OnEventPostExecute(Executor, InEventHandle);
+						}
+					}
+				}
+			}
 
-void FSubsonicSubsystemImpl::StopEvent(FSubsonicEventInstance EventInstance)
-{
-	FSubsonicEventInstanceImpl* ActiveInstance = ActiveInstances.Find(EventInstance);
-	if (ActiveInstance)
-	{
-		ActiveInstance->Stop();
-	}
-}
+			void OnEventSubscriberRegistered(TScriptInterface<ISubsonicEventSubscriberInterface> EventSubscriber)
+			{
+				EventSubscribers.AddUnique(EventSubscriber);
+			}
 
-void FSubsonicSubsystemImpl::RegisterHandlers(USubsonicTempDataAsset* Handlers)
-{
-	if (Handlers)
-	{
-		RegisteredHandlers.Add(Handlers);
-	}
-}
+			void OnEventSubscriberUnregistered(TScriptInterface<ISubsonicEventSubscriberInterface> EventSubscriber)
+			{
+				EventSubscribers.Remove(EventSubscriber);
+			}
 
-void USubsonicSubsystem::Initialize(FSubsystemCollectionBase& Collection)
-{
-	Super::Initialize(Collection);
-	Impl = MakePimpl<FSubsonicSubsystemImpl>();
-}
+			void OnExecutorRegistered(const FSubsonicExecutor& InExecutor)
+			{
+				for (TScriptInterface<ISubsonicEventSubscriberInterface>& EventSubscriber : EventSubscribers)
+				{
+					if (const UObject* EventSystemObj = EventSubscriber.GetObject())
+					{
+						EventSubscriber->OnExecutorRegistered(InExecutor);
+					}
+				}
+			}
 
-void USubsonicSubsystem::Deinitialize()
-{
-	Impl.Reset();
-	Super::Deinitialize();
-}
+			void OnExecutorUnregistered(const FSubsonicExecutor& InExecutor)
+			{
+				for (TScriptInterface<ISubsonicEventSubscriberInterface>& EventSubscriber : EventSubscribers)
+				{
+					if (const UObject* EventSystemObj = EventSubscriber.GetObject())
+					{
+						EventSubscriber->OnExecutorUnregistered(InExecutor);
+					}
+				}
+			}
 
-FSubsonicEventCollectionExecutor USubsonicSubsystem::CreateExecutor(const UObject* WorldContextObject, const USubsonicEventCollection* Collection)
-{
-	using namespace UE::Subsonic;
+		private:
+			virtual void AddReferencedObjects(FReferenceCollector& Collector) override
+			{
+				TArray<TObjectPtr<UObject>> Objects;
+				Algo::Transform(EventSubscribers, Objects, [](const TScriptInterface<ISubsonicEventSubscriberInterface>& Subscriber)
+				{
+					return TObjectPtr<UObject>(Subscriber.GetObject());
+				});
+				Collector.AddReferencedObjects(Objects);
+			}
+
+			virtual FString GetReferencerName() const override
+			{
+				return "FSubsonicSubsystemImpl";
+			}
 
-	const Audio::FDeviceId DeviceId = SubsystemPrivate::GetDeviceIdFrom(WorldContextObject);
-	return FSubsonicEventCollectionExecutor(Collection, DeviceId);
-}
+			TArray<TScriptInterface<ISubsonicEventSubscriberInterface>> EventSubscribers;
+		};
 
-void USubsonicSubsystem::ExecuteEvent(const FSubsonicEventCollectionExecutor& Executor, FGameplayTag EventTag)
-{
-	using namespace UE::Subsonic;
+	} // namespace SubsystemPrivate
 
-	if (!Executor.IsValid())
+	void USubsonicSubsystem::Initialize(FSubsystemCollectionBase& Collection)
 	{
-		return;
+		Super::Initialize(Collection);
+		Impl = MakePimpl<SubsystemPrivate::FSubsonicSubsystemImpl>();
 	}
 
-	const USubsonicEventCollection& Collection = Executor.GetCollectionChecked();
+	void USubsonicSubsystem::Deinitialize()
+	{
+		Impl.Reset();
+		Super::Deinitialize();
+	}
 
-	if (const FSubsonicEvent* Event = Collection.CollectionData.Events.Find(EventTag))
+	USubsonicEventCollectionExecutor* USubsonicSubsystem::CreateExecutorBP(UObject* WorldContextObject, const USubsonicEventCollection* Collection)
 	{
-		const UE::Subsonic::FEventHandle Handle { .Collection = Collection.GetHandle(Executor), .EventName = EventTag.GetTagName() };
+		USubsonicEventCollectionExecutor* NewExecutor = NewObject<USubsonicEventCollectionExecutor>(WorldContextObject);
 
-		// Notify associated device event systems event about to begin
-		TArray<TScriptInterface<ISubsonicEventSystemInterface>> NotifiedSystems;
-		for (const TScriptInterface<ISubsonicEventSystemInterface>& EventSystem : EventSystems)
+		if (GEngine && NewExecutor && Collection)
 		{
-			const UAudioEngineSubsystem* Subsystem = Cast<UAudioEngineSubsystem>(EventSystem.GetObject());
-			const Audio::FDeviceId SystemDeviceId = Subsystem
-				? Subsystem->GetAudioDeviceId()
-				: SubsystemPrivate::GetDeviceIdFrom(EventSystem.GetObject());
+			const Audio::FDeviceId DeviceId = SubsystemPrivate::GetDeviceIdFromWorldObject(WorldContextObject);
 
-			if (SystemDeviceId == Handle.Collection.DeviceId)
-			{
-				if (EventSystem.GetObject())
-				{
-					NotifiedSystems.Add(EventSystem);
-					EventSystem->OnEventPreExecute(Handle);
-				}
-			}
+			NewExecutor->Collection = Collection;
+
+			FCollectionHandle CollectionHandle = Collection->GetHandle();
+			NewExecutor->Executor = FSubsonicExecutor::Create(DeviceId, MoveTemp(CollectionHandle));
+
+			return NewExecutor;
 		}
 
-		Event->Execute(Handle);
+		return nullptr;
+	}
 
-		// Notify associated device event systems when event is complete
-		for (TScriptInterface<ISubsonicEventSystemInterface>& System : NotifiedSystems)
+	void USubsonicSubsystem::ExecuteEventBP(const USubsonicEventCollectionExecutor* Executor, FGameplayTag EventTag)
+	{
+		if (Executor)
 		{
-			check(System.GetObject());
-			System->OnEventPostExecute(Handle);
+			const USubsonicEventCollection& EventCollection = Executor->GetCollectionChecked();
+			const FSubsonicEventCollectionDefinition& Definition = EventCollection.GetDefinition();
+			if (const FSubsonicEvent* Event = Definition.GetEvents().Find(EventTag))
+			{
+				const FCollectionHandle CollectionHandle = EventCollection.GetHandle();
+				Executor->GetExecutor().ExecuteEvent(CollectionHandle, EventTag.GetTagName(), *Event);
+			}
 		}
 	}
-}
 
-FSubsonicEventInstance USubsonicSubsystem::PlayEvent(const UObject* WorldContextObject, const FName EventName)
-{
-	return Impl->PlayEvent(WorldContextObject, EventName);
-
-}
-
-void USubsonicSubsystem::StopEvent(FSubsonicEventInstance EventInstance)
-{
-	Impl->StopEvent(EventInstance);
-}
+	void USubsonicSubsystem::OnEventPreExecute(const FSubsonicExecutor& Executor, const FEventHandle& InEventHandle)
+	{
+		Impl->OnEventPreExecute(Executor, InEventHandle);
+	}
 
-void USubsonicSubsystem::RegisterHandlers(USubsonicTempDataAsset* Handlers)
-{
-	Impl->RegisterHandlers(Handlers);
-}
+	void USubsonicSubsystem::OnEventPostExecute(const FSubsonicExecutor& Executor, const FEventHandle& InEventHandle)
+	{
+		Impl->OnEventPostExecute(Executor, InEventHandle);
+	}
 
-void USubsonicSubsystem::RegisterCollection(const UE::Subsonic::FCollectionHandle& InHandle)
-{
-	using namespace UE::Subsonic;
+	void USubsonicSubsystem::OnCollectionRegistered(const FCollectionHandle& InHandle)
+	{
+		Impl->OnCollectionRegistered(InHandle);
+	}
 
-	FCollectionHandle CollectionHandle = InHandle;
-	for (TScriptInterface<ISubsonicEventSystemInterface>& EventSystem : EventSystems)
+	void USubsonicSubsystem::OnCollectionUnregistered(const FCollectionHandle& InHandle)
 	{
-		if (const UObject* EventSystemObj = EventSystem.GetObject())
-		{
-			CollectionHandle.DeviceId = SubsystemPrivate::GetDeviceIdFrom(EventSystemObj);
-			EventSystem->OnCollectionRegistered(CollectionHandle);
-		}
+		Impl->OnCollectionUnregistered(InHandle);
 	}
-}
 
-void USubsonicSubsystem::RegisterEventSystem(TScriptInterface<ISubsonicEventSystemInterface> EventSystem)
-{
-	EventSystems.AddUnique(EventSystem);
-}
+	void USubsonicSubsystem::OnEventSubscriberRegistered(TScriptInterface<ISubsonicEventSubscriberInterface> EventSubscriber)
+	{
+		Impl->OnEventSubscriberRegistered(EventSubscriber);
+	}
 
-void USubsonicSubsystem::UnregisterCollection(const UE::Subsonic::FCollectionHandle& InHandle)
-{
-	using namespace UE::Subsonic;
+	void USubsonicSubsystem::OnEventSubscriberUnregistered(TScriptInterface<ISubsonicEventSubscriberInterface> EventSubscriber)
+	{
+		Impl->OnEventSubscriberUnregistered(EventSubscriber);
+	}
 
-	FCollectionHandle CollectionHandle = InHandle;
-	for (TScriptInterface<ISubsonicEventSystemInterface>& EventSystem : EventSystems)
+	void USubsonicSubsystem::OnExecutorRegistered(const FSubsonicExecutor& InExecutor)
 	{
-		if (const UObject* EventSystemObj = EventSystem.GetObject())
-		{
-			CollectionHandle.DeviceId = SubsystemPrivate::GetDeviceIdFrom(EventSystemObj);
-			EventSystem->OnCollectionUnregistered(CollectionHandle);
-		}
+		Impl->OnExecutorRegistered(InExecutor);
 	}
-}
 
-void USubsonicSubsystem::UnregisterEventSystem(TScriptInterface<ISubsonicEventSystemInterface> EventSystem)
-{
-	EventSystems.Remove(EventSystem);
-}
+	void USubsonicSubsystem::OnExecutorUnregistered(const FSubsonicExecutor& InExecutor)
+	{
+		Impl->OnExecutorUnregistered(InExecutor);
+	}
+} // namespace UE::Subsonic
diff --git a/Engine/Plugins/Experimental/Subsonic/Source/Subsonic/Private/SubsonicTempDataAsset.h b/Engine/Plugins/Experimental/Subsonic/Source/Subsonic/Private/SubsonicTempDataAsset.h
deleted file mode 100644
diff --git a/Engine/Plugins/Experimental/Subsonic/Source/Subsonic/Private/SubsonicTypes.cpp b/Engine/Plugins/Experimental/Subsonic/Source/Subsonic/Private/SubsonicTypes.cpp
--- a/Engine/Plugins/Experimental/Subsonic/Source/Subsonic/Private/SubsonicTypes.cpp
+++ b/Engine/Plugins/Experimental/Subsonic/Source/Subsonic/Private/SubsonicTypes.cpp
@@ -2,7 +2,10 @@
 
 #include "SubsonicTypes.h"
 
-uint32 GetTypeHash(FSubsonicEventInstance Instance)
+
+namespace UE::Subsonic
 {
-	return GetTypeHash(Instance.RawHandle);
-}
+// 	uint32 FSubsonicEventInstance::GetTypeHash(FSubsonicEventInstance Instance)
+// 	{
+// 	}
+} // namespace UE::Subsonic
\ No newline at end of file
diff --git a/Engine/Plugins/Experimental/Subsonic/Source/Subsonic/Public/CoreEventSubscribers/SubsonicEventSubscriberBase.h b/Engine/Plugins/Experimental/Subsonic/Source/Subsonic/Public/CoreEventSubscribers/SubsonicEventSubscriberBase.h
new file mode 100644
--- /dev/null
+++ b/Engine/Plugins/Experimental/Subsonic/Source/Subsonic/Public/CoreEventSubscribers/SubsonicEventSubscriberBase.h
@@ -0,0 +1,30 @@
+// Copyright Epic Games, Inc. All Rights Reserved.
+#pragma once
+
+#include "SubsonicEventSystemInterface.h"
+#include "Subsystems/AudioEngineSubsystem.h"
+
+#include "SubsonicEventSubscriberBase.generated.h"
+
+#define UE_API SUBSONIC_API
+
+
+namespace UE::Subsonic
+{
+	// Forward Declarations
+	class USubsonicSubsystem;
+
+	UCLASS(MinimalAPI, Abstract)
+		class USubsonicEventSubscriberBase : public UAudioEngineSubsystem, public ISubsonicEventSubscriberInterface
+	{
+		GENERATED_BODY()
+
+	public:
+		void Initialize(FSubsystemCollectionBase& Collection);
+		void Deinitialize();
+
+	protected:
+		UE_API USubsonicSubsystem* GetSubsonicSubsystem() const;
+	};
+} // namespace UE::Subsonic
+#undef UE_API // SUBSONIC_API
\ No newline at end of file
diff --git a/Engine/Plugins/Experimental/Subsonic/Source/Subsonic/Public/SubsonicEventCollection.h b/Engine/Plugins/Experimental/Subsonic/Source/Subsonic/Public/SubsonicEventCollection.h
--- a/Engine/Plugins/Experimental/Subsonic/Source/Subsonic/Public/SubsonicEventCollection.h
+++ b/Engine/Plugins/Experimental/Subsonic/Source/Subsonic/Public/SubsonicEventCollection.h
@@ -5,7 +5,7 @@
 #include "GameplayTagContainer.h"
 #include "NativeGameplayTags.h"
 #include "StructUtils/InstancedStruct.h"
-#include "SubsonicHandles.h"
+#include "SubsonicExecutor.h"
 #include "UObject/NameTypes.h"
 
 #include "SubsonicEventCollection.generated.h"
@@ -19,81 +19,168 @@ namespace UE::Subsonic
 
 	UE_API UE_DECLARE_GAMEPLAY_TAG_EXTERN(TAG_SubsonicEventCorePlay)
 	UE_API UE_DECLARE_GAMEPLAY_TAG_EXTERN(TAG_SubsonicEventCoreStop)
-} // namespace UE::Subsonic
-
-
-USTRUCT(MinimalAPI, BlueprintType, meta = (Abstract, Hidden))
-struct FSubsonicEventActionBase
-{
-	GENERATED_BODY()
-
-	virtual ~FSubsonicEventActionBase() = default;
-
-	virtual void Execute(const UE::Subsonic::FActionHandle& InHandle) const { }
-};
-
-USTRUCT(MinimalAPI, BlueprintType)
-struct FSubsonicEvent
-{
-	GENERATED_BODY()
-
-	virtual ~FSubsonicEvent() = default;
-
-	void Execute(const UE::Subsonic::FEventHandle& InHandle) const;
-
-	UPROPERTY(EditAnywhere, Category = "Subsonic")
-	TArray<TInstancedStruct<FSubsonicEventActionBase>> Actions;
-};
-
-USTRUCT(MinimalAPI, BlueprintType)
-struct FSubsonicEventCollectionDefinition
-{
-	GENERATED_BODY()
 
-public:
-	UPROPERTY(EditAnywhere, Category = "Subsonic")
-	TMap<FGameplayTag, FSubsonicEvent> Events;
-};
+	USTRUCT(MinimalAPI, BlueprintType, meta = (Abstract, Hidden))
+	struct FSubsonicEventActionBase
+	{
+		GENERATED_BODY()
 
+		virtual ~FSubsonicEventActionBase() = default;
 
-UCLASS(MinimalAPI, BlueprintType)
-class USubsonicEventCollection final : public UObject
-{
-	GENERATED_BODY()
+	protected:
+		virtual void Execute(const FSubsonicExecutor& InExecutor, const FActionHandle& InHandle) const { }
 
-public:
-	UPROPERTY(EditAnywhere, Category = "Subsonic", meta = (ShowOnlyInnerProperties))
-	FSubsonicEventCollectionDefinition CollectionData;
+		friend struct FSubsonicEvent;
+	};
 
-	UE::Subsonic::FCollectionHandle GetHandle(const FSubsonicEventCollectionExecutor& Executor) const;
+	USTRUCT(MinimalAPI, BlueprintType)
+	struct FSubsonicEvent
+	{
+		GENERATED_BODY()
 
-	virtual void PostLoad() override;
-	virtual void BeginDestroy() override;
-};
+		virtual ~FSubsonicEvent() = default;
+
+	private:
+		friend struct FSubsonicExecutor;
+
+		// Called by executor, iterates through all actions on the given event definition.
+		// Asserts if the provided executor and handle have mismatched collection contexts.
+		UE_INTERNAL void Execute(const FSubsonicExecutor& InExecutor, const FEventHandle& InHandle) const;
+
+	public:
+		UPROPERTY(EditAnywhere, Category = "Subsonic")
+		TArray<TInstancedStruct<FSubsonicEventActionBase>> Actions;
+	};
+
+
+	USTRUCT(MinimalAPI, BlueprintType)
+	struct FSubsonicEventCollectionDefinition final
+	{
+		GENERATED_BODY()
+
+		// Generates empty, invalid event collection. Required to be public
+		// to support UObject reflection. To produce a valid definition natively
+		// without UObject support, call static Create function below.
+		FSubsonicEventCollectionDefinition() = default;
+
+		FSubsonicEventCollectionDefinition(FSubsonicEventCollectionDefinition&& Definition);
+		FSubsonicEventCollectionDefinition(const FSubsonicEventCollectionDefinition& Definition) = delete;
+
+		FSubsonicEventCollectionDefinition& operator =(FSubsonicEventCollectionDefinition&& Definition);
+		FSubsonicEventCollectionDefinition& operator =(const FSubsonicEventCollectionDefinition& Definition) = delete;
+
+		~FSubsonicEventCollectionDefinition();
+
+		// Create a collection definition with a valid transient Id for tracking
+		// within the Subsonic System. Used to create and register definitions
+		// outside of UObject reflection.
+		static FSubsonicEventCollectionDefinition Create(FName Name, TMap<FGameplayTag, FSubsonicEvent> InEvents);
+
+		UE_API bool IsValid() const;
+		UE_API uint32 GetCollectionId() const;
+		UE_API const TMap<FGameplayTag, FSubsonicEvent>& GetEvents() const;
+
+		// Returns Id of invalid (default constructed or unregistered/reset collection).
+		UE_API static uint32 GetInvalidId();
+
+	private:
+		// Register's the given collection with the Subsonic Subsystem.
+		UE_INTERNAL void Register(const FCollectionHandle& InHandle);
+
+		// Unregisters collection, leaving it in clear, invalid state.
+		UE_INTERNAL void Unregister();
+
+		// Unique Collection Id assigned post
+		// initialization of struct properties
+		// or via static initializer
+		// Independent from UObject identifier to
+		// allow for transient construction without
+		// having to utilize UObject wrapper.
+		uint32 CollectionId = INDEX_NONE;
+
+		UPROPERTY(EditAnywhere, BlueprintReadOnly, Category = "Subsonic", meta = (AllowPrivateAccess))
+		TMap<FGameplayTag, FSubsonicEvent> Events;
+
+		friend class USubsonicEventCollection;
+
+		void AssignId();
+	};
+
+	UCLASS(MinimalAPI, BlueprintType)
+	class USubsonicEventCollection final : public UObject
+	{
+		GENERATED_BODY()
+
+	public:
+		const FSubsonicEventCollectionDefinition& GetDefinition() const
+		{
+			return CollectionDefinition;
+		}
+
+		FCollectionHandle GetHandle() const
+		{
+	#if WITH_EDITORONLY_DATA
+			FName CollectionName;
+			if (const UPackage* CollectionPackage = GetPackage())
+			{
+				const FTopLevelAssetPath AssetPath(CollectionPackage);
+				CollectionName = AssetPath.GetPackageName();
+			}
+	#endif // WITH_EDITORONLY_DATA
+
+			return FCollectionHandle
+			{
+	#if WITH_EDITORONLY_DATA
+				.CollectionName = CollectionName,
+	#endif // WITH_EDITORONLY_DATA
+				.CollectionId = CollectionDefinition.GetCollectionId(),
+			};
+		}
+
+		void SetDefinition(FSubsonicEventCollectionDefinition Definition)
+		{
+			Definition.AssignId();
+			CollectionDefinition = MoveTemp(Definition);
+		}
+
+	private:
+		UPROPERTY(EditAnywhere, Category = "Subsonic", meta = (ShowOnlyInnerProperties))
+		FSubsonicEventCollectionDefinition CollectionDefinition;
+
+		virtual void PostInitProperties() override;
+		virtual void PostLoad() override;
+		virtual void BeginDestroy() override;
+	};
+
+	// Executor handle for a Subsonic Event Collection.
+	UCLASS(MinimalAPI, BlueprintType)
+	class USubsonicEventCollectionExecutor final : public UObject
+	{
+		GENERATED_BODY()
+
+	public:
+		const USubsonicEventCollection* GetCollection() const;
+		const USubsonicEventCollection& GetCollectionChecked() const;
+		const FSubsonicExecutor& GetExecutor() const;
+		bool IsValid() const;
+
+		virtual void BeginDestroy() override;
+
+	private:
+		UPROPERTY(Transient, BlueprintReadOnly, Category = "Subsonic", meta = (AllowPrivateAccess))
+		TObjectPtr<const USubsonicEventCollection> Collection;
+
+		FSubsonicExecutor Executor = FSubsonicExecutor::GetInvalid();
+
+		friend class USubsonicSubsystem;
+	};
+} // namespace UE::Subsonic
 
-USTRUCT(MinimalAPI, BlueprintType)
-struct FSubsonicEventCollectionExecutor final
+// Required to be in global namespace still for UHT compatibility with namespaced Unreal reflected types
+template <>
+struct TStructOpsTypeTraits<UE::Subsonic::FSubsonicEventCollectionDefinition>
+	: public TStructOpsTypeTraitsBase2<UE::Subsonic::FSubsonicEventCollectionDefinition>
 {
-	GENERATED_BODY()
-
-private:
-	static uint32 MaxExecutorId;
-
-public:
-	FSubsonicEventCollectionExecutor() = default;
-	FSubsonicEventCollectionExecutor(const USubsonicEventCollection* InCollection, Audio::FDeviceId InDeviceId);
-
-	Audio::FDeviceId GetDeviceId() const;
-	const USubsonicEventCollection* GetCollection() const;
-	const USubsonicEventCollection& GetCollectionChecked() const;
-	uint32 GetExecutorId() const;
-	bool IsValid() const;
-
-private:
-	UPROPERTY(BlueprintReadOnly, Category = "Subsonic", meta = (AllowPrivateAccess))
-	TObjectPtr<const USubsonicEventCollection> Collection;
-
-	Audio::FDeviceId DeviceId = INDEX_NONE;
-	uint32 ExecutorId = INDEX_NONE;
+	enum { WithCopy = false };
 };
 #undef UE_API // SUBSONIC_API
diff --git a/Engine/Plugins/Experimental/Subsonic/Source/Subsonic/Public/SubsonicExecutor.h b/Engine/Plugins/Experimental/Subsonic/Source/Subsonic/Public/SubsonicExecutor.h
new file mode 100644
--- /dev/null
+++ b/Engine/Plugins/Experimental/Subsonic/Source/Subsonic/Public/SubsonicExecutor.h
@@ -0,0 +1,147 @@
+// Copyright Epic Games, Inc. All Rights Reserved.
+#pragma once
+
+#include "AudioDefines.h"
+#include "AudioDevice.h"
+#include "HAL/Platform.h"
+#include "Misc/CoreDefines.h"
+#include "SubsonicHandles.h"
+#include "Templates/TypeHash.h"
+#include "UObject/NameTypes.h"
+
+#include "SubsonicExecutor.generated.h"
+
+
+// Forward Declarations
+struct FSubsonicEvent;
+struct FSubsonicEventCollectionDefinition;
+
+#define UE_API SUBSONIC_API
+
+
+UENUM(BlueprintType)
+enum class ESubsonicExecutionScope : uint8
+{
+	Global,
+	Collection,
+	Event,
+	Action
+};
+
+
+namespace UE::Subsonic
+{
+	// Execution object that tracks explicit execution context for any caller of events related to a given collection.
+	// RAII design that automatically registers and unregisters with the Subsonic Subsystem, and by extension,
+	// any implemented and active subscribers.  This enables subscribers to track and manage data related to
+	// event requests from a given executor instance in a data-oriented fashion.
+	struct FSubsonicExecutor final
+	{
+	private:
+		FSubsonicExecutor();
+		FSubsonicExecutor(const FSubsonicExecutor&) = delete;
+
+	public:
+		FSubsonicExecutor(FSubsonicExecutor&&);
+		FSubsonicExecutor& operator=(const FSubsonicExecutor& InOther) = delete;
+		virtual ~FSubsonicExecutor();
+
+		UE_API FSubsonicExecutor& operator=(FSubsonicExecutor&& InOther);
+
+		UE_API static FSubsonicExecutor Create(Audio::FDeviceId InDeviceId, FCollectionHandle CollectionHandle);
+		UE_API static uint32 GetInvalidId();
+		UE_API static FSubsonicExecutor GetInvalid();
+
+		FCollectionHandle GetCollectionHandle() const;
+		Audio::FDeviceId GetDeviceId() const;
+		uint32 GetId() const;
+
+		void ExecuteEvent(const FCollectionHandle& InCollection, FName EventName, const FSubsonicEvent& InEvent) const;
+
+		UE_API bool IsValid() const;
+
+		UE_API FString ToString() const;
+
+		// Unregisters Executor from the SubsonicSubsystem (and by extension all subscribers), reporting
+		// to and releasing all subscriber data keyed off the given executor. Id is cleared, leaving
+		// local executor in invalid state.
+		UE_API void Unregister();
+
+	private:
+		// Registers Executor with the SubsonicSubsystem (and by extension all subscribers).
+		void Register();
+
+		FCollectionHandle CollectionHandle;
+
+		// DeviceId corresponding to which audio device to execute event on.
+		Audio::FDeviceId DeviceId;
+
+		// Id of calling context executor.
+		uint32 ExecutorId = GetInvalidId();
+	};
+
+	// Key consisting of executor and scoped handle (ex. CollectionHandle, EventHandle, ActionHandle).
+	// Essentially a pair of both values that is easily readable, hashable, and provided as convenience
+	// functionality for subscribers to organize data based in hashed collection related to execution context.
+	template <typename ScopeHandleType>
+	struct TScopeKey
+	{
+		uint32 ExecutorId = FSubsonicExecutor::GetInvalidId();
+		Audio::FDeviceId DeviceId = INDEX_NONE;
+		ScopeHandleType ScopeHandle;
+
+		TScopeKey() = default;
+		TScopeKey(const FSubsonicExecutor& InExecutor, ScopeHandleType InScopeHandle)
+			: ExecutorId(InExecutor.GetId())
+			, DeviceId(InExecutor.GetDeviceId())
+			, ScopeHandle(MoveTemp(InScopeHandle))
+		{
+		}
+
+		template <typename SubsystemType>
+		SubsystemType* FindDeviceSubsystem() const
+		{
+			const FAudioDevice* AudioDevice = nullptr;
+			if (const FAudioDeviceManager* DeviceManager = FAudioDeviceManager::Get())
+			{
+				if (DeviceId == INDEX_NONE)
+				{
+					AudioDevice = DeviceManager->GetMainAudioDeviceRaw();
+				}
+				else
+				{
+					AudioDevice = DeviceManager->GetAudioDeviceRaw(DeviceId);
+				}
+			}
+
+			if (AudioDevice)
+			{
+				return AudioDevice->GetSubsystem<SubsystemType>();
+			}
+
+			return nullptr;
+		}
+
+		friend uint32 GetTypeHash(const TScopeKey<ScopeHandleType>& Key)
+		{
+			const uint32 Hash1 = HashCombineFast(Key.ExecutorId, GetTypeHash(Key.DeviceId));
+			return HashCombineFast(Hash1, GetTypeHash(Key.ScopeHandle));
+		}
+
+		friend bool operator==(const TScopeKey<ScopeHandleType>& A, const TScopeKey<ScopeHandleType>& B)
+		{
+			return A.ExecutorId == B.ExecutorId && A.ScopeHandle == B.ScopeHandle && A.DeviceId == B.DeviceId;
+		}
+
+		FString ToString() const
+		{
+			return FString::Printf(TEXT("%s [ExecId: %u, AudioDevice Id: %u]"), *ScopeHandle.ToString(), ExecutorId, DeviceId);
+		}
+	};
+
+	using FCollectionScopeKey = TScopeKey<FCollectionHandle>;
+	using FEventScopeKey = TScopeKey<FEventHandle>;
+	using FActionScopeKey = TScopeKey<FActionHandle>;
+} // namespace UE::Subsonic
+
+#undef UE_API
diff --git a/Engine/Plugins/Experimental/Subsonic/Source/Subsonic/Public/SubsonicHandles.h b/Engine/Plugins/Experimental/Subsonic/Source/Subsonic/Public/SubsonicHandles.h
--- a/Engine/Plugins/Experimental/Subsonic/Source/Subsonic/Public/SubsonicHandles.h
+++ b/Engine/Plugins/Experimental/Subsonic/Source/Subsonic/Public/SubsonicHandles.h
@@ -7,25 +7,16 @@
 #include "Misc/CoreDefines.h"
 #include "Templates/TypeHash.h"
 #include "UObject/NameTypes.h"
-#include "UObject/WeakObjectPtrTemplates.h"
-
-#include "SubsonicHandles.generated.h"
 
 #define UE_API SUBSONIC_API
 
 
-UENUM(BlueprintType)
-enum class ESubsonicExecutionScope : uint8
-{
-	Global,
-	Collection,
-	Event,
-	Action
-};
-
-
 namespace UE::Subsonic
 {
+	// Forward Declarations
+	struct FSubsonicEvent;
+	struct FSubsonicEventCollectionDefinition;
+
 	struct FCollectionHandle
 	{
 #if WITH_EDITORONLY_DATA
@@ -33,64 +24,28 @@ namespace UE::Subsonic
 		FName CollectionName;
 #endif // WITH_EDITORONLY_DATA
 
-		// Owning Collection's instance id
-		uint32 CollectionId = static_cast<uint32>(INDEX_NONE);
-
-		// DeviceId corresponding to which audio device to execute event on.
-		Audio::FDeviceId DeviceId = static_cast<Audio::FDeviceId>(INDEX_NONE);
+		// Returns invalid Id of a handle
+		UE_API static uint32 GetInvalidId();
 
-		// Id of calling context executor. May be unset if handle context does not require it.
-		uint32 ExecutorId = static_cast<uint32>(INDEX_NONE);
-
-		template <typename SubsystemType>
-		SubsystemType* FindDeviceSubsystem() const
-		{
-			if (FAudioDeviceManager* DeviceManager = FAudioDeviceManager::Get())
-			{
-				FAudioDevice* AudioDevice = nullptr;
-				if (DeviceId == INDEX_NONE)
-				{
-					AudioDevice = DeviceManager->GetMainAudioDeviceRaw();
-				}
-				else
-				{
-					AudioDevice = DeviceManager->GetAudioDeviceRaw(DeviceId);
-				}
-
-				if (AudioDevice)
-				{
-					return AudioDevice->GetSubsystem<SubsystemType>();
-				}
-			}
-
-			return nullptr;
-		}
+		// Owning Collection's instance id
+		uint32 CollectionId = GetInvalidId();
 
 		friend uint32 GetTypeHash(const FCollectionHandle& Handle)
 		{
-			const uint32 Hash1 = HashCombineFast(GetTypeHash(Handle.CollectionId), GetTypeHash(Handle.ExecutorId));
-			return HashCombineFast(Hash1, GetTypeHash(Handle.DeviceId));
+			return Handle.CollectionId;
 		}
 
 		friend bool operator==(const FCollectionHandle& A, const FCollectionHandle& B)
 		{
-			return A.ExecutorId == B.ExecutorId
-				&& A.CollectionId == B.CollectionId
-				&& A.DeviceId == B.DeviceId;
+			return A.CollectionId == B.CollectionId;
 		}
 
-		FString ToString() const
-		{
-#if WITH_EDITORONLY_DATA
-			return FString::Printf(TEXT("%s(%u, Exec: %u)"), *CollectionName.ToString(), CollectionId, ExecutorId);
-#else // !WITH_EDITORONLY_DATA
-			return FString::Printf(TEXT("%u, Exec: %u"), CollectionId, ExecutorId);
-#endif // !WITH_EDITORONLY_DATA
-		}
+		FString ToString() const;
 	};
 
 	struct FEventHandle
 	{
+		// Parent collection handle
 		FCollectionHandle Collection;
 
 		// Owning event name
@@ -106,14 +61,12 @@ namespace UE::Subsonic
 			return A.Collection == B.Collection && A.EventName == B.EventName;
 		}
 
-		FString ToString() const
-		{
-			return FString::Printf(TEXT("%s/%s"), *Collection.ToString(), *EventName.ToString());
-		}
+		FString ToString() const;
 	};
 
 	struct FActionHandle
 	{
+		// Parent event handle
 		FEventHandle Event;
 
 		// Owning action's index within a given event
@@ -129,10 +82,7 @@ namespace UE::Subsonic
 			return A.Event == B.Event && A.Index == B.Index;
 		}
 
-		FString ToString() const
-		{
-			return FString::Printf(TEXT("%s[%i]"), *Event.ToString(), Index);
-		}
+		FString ToString() const;
 	};
 } // namespace UE::Subsonic
 
diff --git a/Engine/Plugins/Experimental/Subsonic/Source/Subsonic/Public/SubsonicSubsystem.h b/Engine/Plugins/Experimental/Subsonic/Source/Subsonic/Public/SubsonicSubsystem.h
--- a/Engine/Plugins/Experimental/Subsonic/Source/Subsonic/Public/SubsonicSubsystem.h
+++ b/Engine/Plugins/Experimental/Subsonic/Source/Subsonic/Public/SubsonicSubsystem.h
@@ -3,7 +3,9 @@
 #pragma once
 
 #include "GameplayTagContainer.h"
+#include "SubsonicEventCollection.h"
 #include "SubsonicEventSystemInterface.h"
+#include "SubsonicExecutor.h"
 #include "SubsonicTypes.h"
 #include "Subsystems/EngineSubsystem.h"
 
@@ -11,51 +13,56 @@
 
 #define UE_API SUBSONIC_API
 
-class FSubsonicEventInstanceImpl;
-class FSubsonicSubsystemImpl;
-class USubsonicTempDataAsset;
-class USubsonicEventCollection;
 
-/**
- * Interface to the Subsonic plugin.
- */
-UCLASS(MinimalAPI, BlueprintType)
-class USubsonicSubsystem final : public UEngineSubsystem
+namespace UE::Subsonic
 {
-	GENERATED_BODY()
-
-public:
-	//~ Begin USubsystem interface
-	UE_API virtual void Initialize(FSubsystemCollectionBase& Collection) override;
-	UE_API virtual void Deinitialize() override;
-	//~ End USubsystem interface
-
-	UFUNCTION(BlueprintCallable, Category = "Subsonic")
-	UE_API void RegisterHandlers(USubsonicTempDataAsset* Handlers);
-
-	UFUNCTION(BlueprintCallable, Category = "Subsonic", meta=(WorldContext="WorldContextObject"))
-	UE_API FSubsonicEventInstance PlayEvent(const UObject* WorldContextObject, const FName EventName);
-
-	UFUNCTION(BlueprintCallable, Category = "Subsonic")
-	UE_API void StopEvent(FSubsonicEventInstance EventInstance);
-
-	UE_API void RegisterEventSystem(TScriptInterface<ISubsonicEventSystemInterface> Client);
-	UE_API void RegisterCollection(const UE::Subsonic::FCollectionHandle& InHandle);
-
-	UE_API void UnregisterEventSystem(TScriptInterface<ISubsonicEventSystemInterface> Client);
-	UE_API void UnregisterCollection(const UE::Subsonic::FCollectionHandle& InHandle);
-
-	UFUNCTION(BlueprintCallable, Category = "Subsonic", meta = (WorldContext = "WorldContextObject"))
-	UE_API UPARAM(DisplayName = "Executor") FSubsonicEventCollectionExecutor CreateExecutor(const UObject* WorldContextObject, const USubsonicEventCollection* Collection);
-
-	UFUNCTION(BlueprintCallable, Category = "Subsonic")
-	UE_API void ExecuteEvent(const FSubsonicEventCollectionExecutor& Executor, FGameplayTag Event);
-
-private:
-	TPimplPtr<FSubsonicSubsystemImpl> Impl;
-
-	UPROPERTY()
-	TArray<TScriptInterface<ISubsonicEventSystemInterface>> EventSystems;
-};
-
+	// Forward Declarations
+	class FSubsonicEventInstanceImpl;
+	class USubsonicEventCollection;
+
+	namespace SubsystemPrivate
+	{
+		class FSubsonicSubsystemImpl;
+	} // namespace SubsystemPrivate
+
+	/**
+	 * Interface to the Subsonic plugin.
+	 */
+	UCLASS(MinimalAPI, BlueprintType)
+	class USubsonicSubsystem final : public UEngineSubsystem
+	{
+		GENERATED_BODY()
+
+	public:
+		//~ Begin USubsystem interface
+		UE_API virtual void Initialize(FSubsystemCollectionBase& Collection) override;
+		UE_API virtual void Deinitialize() override;
+		//~ End USubsystem interface
+
+		UFUNCTION(BlueprintCallable, Category = "Subsonic", meta = (DisplayName = "Create Executor", WorldContext = "WorldContextObject"))
+		UE_API UPARAM(DisplayName = "Executor") USubsonicEventCollectionExecutor* CreateExecutorBP(UObject* WorldContextObject, const USubsonicEventCollection* Collection);
+
+		UFUNCTION(BlueprintCallable, Category = "Subsonic", meta = (DisplayName = "Execute Event"))
+		UE_API void ExecuteEventBP(const USubsonicEventCollectionExecutor* Executor, FGameplayTag Event);
+
+	private:
+		friend struct FSubsonicExecutor;
+		friend struct FSubsonicEventCollectionDefinition;
+		friend class USubsonicEventSubscriberBase;
+
+		UE_INTERNAL void OnCollectionRegistered(const FCollectionHandle& InHandle);
+		UE_INTERNAL void OnCollectionUnregistered(const FCollectionHandle& InHandle);
+
+		UE_INTERNAL void OnEventPreExecute(const FSubsonicExecutor& Executor, const FEventHandle& InEventHandle);
+		UE_INTERNAL void OnEventPostExecute(const FSubsonicExecutor& Executor, const FEventHandle& InEventHandle);
+
+		UE_INTERNAL void OnEventSubscriberRegistered(TScriptInterface<ISubsonicEventSubscriberInterface> Subscriber);
+		UE_INTERNAL void OnEventSubscriberUnregistered(TScriptInterface<ISubsonicEventSubscriberInterface> Subscriber);
+
+		UE_INTERNAL void OnExecutorRegistered(const FSubsonicExecutor& InExecutor);
+		UE_INTERNAL void OnExecutorUnregistered(const FSubsonicExecutor& InExecutor);
+
+		TPimplPtr<SubsystemPrivate::FSubsonicSubsystemImpl> Impl;
+	};
+} // namespace UE::Subsonic
 #undef UE_API
diff --git a/Engine/Plugins/Experimental/Subsonic/Source/Subsonic/Public/SubsonicTypes.h b/Engine/Plugins/Experimental/Subsonic/Source/Subsonic/Public/SubsonicTypes.h
--- a/Engine/Plugins/Experimental/Subsonic/Source/Subsonic/Public/SubsonicTypes.h
+++ b/Engine/Plugins/Experimental/Subsonic/Source/Subsonic/Public/SubsonicTypes.h
@@ -6,22 +6,29 @@
 
 #include "SubsonicTypes.generated.h"
 
-/**
-* Handle to a playing instance of a Subsonic event.
-*/
-USTRUCT(BlueprintType)
-struct FSubsonicEventInstance final
-{
-	GENERATED_USTRUCT_BODY();
 
-	bool operator==(const FSubsonicEventInstance& Other) const
+namespace UE::Subsonic
+{
+	/**
+	* Handle to a playing instance of a Subsonic event.
+	*/
+	USTRUCT(BlueprintType)
+	struct FSubsonicEventInstance final
 	{
-		return RawHandle == Other.RawHandle;
-	}
+		GENERATED_USTRUCT_BODY();
+
+		bool operator==(const FSubsonicEventInstance& Other) const
+		{
+			return RawHandle == Other.RawHandle;
+		}
 
-private:
-	uint32 RawHandle = 0;
+	private:
+		uint32 RawHandle = 0;
 
-	friend class FSubsonicSubsystemImpl;
-	friend uint32 GetTypeHash(FSubsonicEventInstance Instance);
-};
+		friend class FSubsonicSubsystemImpl;
+		friend uint32 GetTypeHash(FSubsonicEventInstance Instance)
+		{
+			return GetTypeHash(Instance.RawHandle);
+		}
+	};
+} // namespace UE::Subsonic
\ No newline at end of file
diff --git a/Engine/Plugins/Experimental/Subsonic/Source/Subsonic/Subsonic.Build.cs b/Engine/Plugins/Experimental/Subsonic/Source/Subsonic/Subsonic.Build.cs
--- a/Engine/Plugins/Experimental/Subsonic/Source/Subsonic/Subsonic.Build.cs
+++ b/Engine/Plugins/Experimental/Subsonic/Source/Subsonic/Subsonic.Build.cs
@@ -25,5 +25,7 @@ public class Subsonic : ModuleRules
 				"Engine"
 			}
 		);
+
+		bAllowUETypesInNamespaces = true;
 	}
 }
diff --git a/Engine/Plugins/Experimental/Subsonic/Source/SubsonicEditor/Private/SubsonicAssetDefinitions.cpp b/Engine/Plugins/Experimental/Subsonic/Source/SubsonicEditor/Private/SubsonicAssetDefinitions.cpp
--- a/Engine/Plugins/Experimental/Subsonic/Source/SubsonicEditor/Private/SubsonicAssetDefinitions.cpp
+++ b/Engine/Plugins/Experimental/Subsonic/Source/SubsonicEditor/Private/SubsonicAssetDefinitions.cpp
@@ -11,66 +11,68 @@
 #define LOCTEXT_NAMESPACE "SubsonicEditor"
 
 
-FLinearColor UAssetDefinition_SubsonicEventCollection::GetAssetColor() const
+namespace UE::Subsonic
 {
- 	if (const ISlateStyle* SubsonicStyle = FSlateStyleRegistry::FindSlateStyle("SubsonicStyle"))
- 	{
- 		return SubsonicStyle->GetColor("SubsonicEventCollection.Color").ToFColorSRGB();
- 	}
+	FLinearColor UAssetDefinition_SubsonicEventCollection::GetAssetColor() const
+	{
+ 		if (const ISlateStyle* SubsonicStyle = FSlateStyleRegistry::FindSlateStyle("SubsonicStyle"))
+ 		{
+ 			return SubsonicStyle->GetColor("SubsonicEventCollection.Color").ToFColorSRGB();
+ 		}
  
- 	return FColor::White;
-}
-
-TSoftClassPtr<UObject> UAssetDefinition_SubsonicEventCollection::GetAssetClass() const
-{
-	return USubsonicEventCollection::StaticClass();
-}
+ 		return FColor::White;
+	}
 
-TConstArrayView<FAssetCategoryPath> UAssetDefinition_SubsonicEventCollection::GetAssetCategories() const
-{
-	static const TArray<FAssetCategoryPath> Categories = { EAssetCategoryPaths::Audio / LOCTEXT("AssetSoundSubsonicSubMenu", "Subsonic") };
-	return Categories;
-}
+	TSoftClassPtr<UObject> UAssetDefinition_SubsonicEventCollection::GetAssetClass() const
+	{
+		return USubsonicEventCollection::StaticClass();
+	}
 
-EAssetCommandResult UAssetDefinition_SubsonicEventCollection::OpenAssets(const FAssetOpenArgs& OpenArgs) const
-{
-	using namespace UE::Subsonic::Editor;
+	TConstArrayView<FAssetCategoryPath> UAssetDefinition_SubsonicEventCollection::GetAssetCategories() const
+	{
+		static const TArray<FAssetCategoryPath> Categories = { EAssetCategoryPaths::Audio / LOCTEXT("AssetSoundSubsonicSubMenu", "Subsonic") };
+		return Categories;
+	}
 
-	for (const FAssetData& AssetData : OpenArgs.Assets)
+	EAssetCommandResult UAssetDefinition_SubsonicEventCollection::OpenAssets(const FAssetOpenArgs& OpenArgs) const
 	{
-		if (USubsonicEventCollection* Collection = Cast<USubsonicEventCollection>(AssetData.GetAsset()))
+		using namespace Editor;
+
+		for (const FAssetData& AssetData : OpenArgs.Assets)
 		{
-			TSharedRef<FEventCollectionEditor> NewEditor = MakeShared<FEventCollectionEditor>();
-			NewEditor->Init(OpenArgs.GetToolkitMode(), OpenArgs.ToolkitHost, *Collection);
+			if (USubsonicEventCollection* Collection = Cast<USubsonicEventCollection>(AssetData.GetAsset()))
+			{
+				TSharedRef<FEventCollectionEditor> NewEditor = MakeShared<FEventCollectionEditor>();
+				NewEditor->Init(OpenArgs.GetToolkitMode(), OpenArgs.ToolkitHost, *Collection);
+			}
 		}
-	}
 
-	return EAssetCommandResult::Handled;
-}
-
-const FSlateBrush* UAssetDefinition_SubsonicEventCollection::GetThumbnailBrush(const FAssetData& InAssetData, const FName InClassName) const
-{
-	return Super::GetThumbnailBrush(InAssetData, InClassName);
-}
+		return EAssetCommandResult::Handled;
+	}
 
-const FSlateBrush* UAssetDefinition_SubsonicEventCollection::GetIconBrush(const FAssetData& InAssetData, const FName InClassName) const
-{
-	return Super::GetIconBrush(InAssetData, InClassName);
-}
+	const FSlateBrush* UAssetDefinition_SubsonicEventCollection::GetThumbnailBrush(const FAssetData& InAssetData, const FName InClassName) const
+	{
+		return Super::GetThumbnailBrush(InAssetData, InClassName);
+	}
 
-bool UAssetDefinition_SubsonicEventCollection::GetThumbnailActionOverlay(const FAssetData& InAssetData, FAssetActionThumbnailOverlayInfo& OutActionOverlayInfo) const
-{
-	return Super::GetThumbnailActionOverlay(InAssetData, OutActionOverlayInfo);
-}
+	const FSlateBrush* UAssetDefinition_SubsonicEventCollection::GetIconBrush(const FAssetData& InAssetData, const FName InClassName) const
+	{
+		return Super::GetIconBrush(InAssetData, InClassName);
+	}
 
-EAssetCommandResult UAssetDefinition_SubsonicEventCollection::ActivateAssets(const FAssetActivateArgs& ActivateArgs) const
-{
-	return Super::ActivateAssets(ActivateArgs);
-}
+	bool UAssetDefinition_SubsonicEventCollection::GetThumbnailActionOverlay(const FAssetData& InAssetData, FAssetActionThumbnailOverlayInfo& OutActionOverlayInfo) const
+	{
+		return Super::GetThumbnailActionOverlay(InAssetData, OutActionOverlayInfo);
+	}
 
-void UAssetDefinition_SubsonicEventCollection::GetAssetActionButtonExtensions(const FAssetData& InAssetData, TArray<FAssetButtonActionExtension>& OutExtensions) const
-{
-	Super::GetAssetActionButtonExtensions(InAssetData, OutExtensions);
-}
+	EAssetCommandResult UAssetDefinition_SubsonicEventCollection::ActivateAssets(const FAssetActivateArgs& ActivateArgs) const
+	{
+		return Super::ActivateAssets(ActivateArgs);
+	}
 
+	void UAssetDefinition_SubsonicEventCollection::GetAssetActionButtonExtensions(const FAssetData& InAssetData, TArray<FAssetButtonActionExtension>& OutExtensions) const
+	{
+		Super::GetAssetActionButtonExtensions(InAssetData, OutExtensions);
+	}
+} // namespace UE::Subsonic
 #undef LOCTEXT_NAMESPACE // SubsonicEditor
diff --git a/Engine/Plugins/Experimental/Subsonic/Source/SubsonicEditor/Private/SubsonicAssetDefinitions.h b/Engine/Plugins/Experimental/Subsonic/Source/SubsonicEditor/Private/SubsonicAssetDefinitions.h
--- a/Engine/Plugins/Experimental/Subsonic/Source/SubsonicEditor/Private/SubsonicAssetDefinitions.h
+++ b/Engine/Plugins/Experimental/Subsonic/Source/SubsonicEditor/Private/SubsonicAssetDefinitions.h
@@ -5,25 +5,27 @@
 
 #include "SubsonicAssetDefinitions.generated.h"
 
-
-UCLASS()
-class UAssetDefinition_SubsonicEventCollection : public UAssetDefinitionDefault_AudioDiffable
+namespace UE::Subsonic
 {
-	GENERATED_BODY()
+	UCLASS()
+	class UAssetDefinition_SubsonicEventCollection : public UAssetDefinitionDefault_AudioDiffable
+	{
+		GENERATED_BODY()
 
-public:
-	// UAssetDefinition Begin
-	virtual FText GetAssetDisplayName() const override { return NSLOCTEXT("AssetTypeActions", "AssetTypeActions_SubsonicEventCollection", "Subsonic Event Collection"); }
-	virtual FLinearColor GetAssetColor() const override;
-	virtual TSoftClassPtr<UObject> GetAssetClass() const override;
-	virtual TConstArrayView<FAssetCategoryPath> GetAssetCategories() const override;
-	virtual EAssetCommandResult OpenAssets(const FAssetOpenArgs& OpenArgs) const override;
+	public:
+		// UAssetDefinition Begin
+		virtual FText GetAssetDisplayName() const override { return NSLOCTEXT("AssetTypeActions", "AssetTypeActions_SubsonicEventCollection", "Subsonic Event Collection"); }
+		virtual FLinearColor GetAssetColor() const override;
+		virtual TSoftClassPtr<UObject> GetAssetClass() const override;
+		virtual TConstArrayView<FAssetCategoryPath> GetAssetCategories() const override;
+		virtual EAssetCommandResult OpenAssets(const FAssetOpenArgs& OpenArgs) const override;
 
-	virtual bool GetThumbnailActionOverlay(const FAssetData& InAssetData, FAssetActionThumbnailOverlayInfo& OutActionOverlayInfo) const override;
-	virtual EAssetCommandResult ActivateAssets(const FAssetActivateArgs& ActivateArgs) const override;
-	virtual void GetAssetActionButtonExtensions(const FAssetData& InAssetData, TArray<FAssetButtonActionExtension>& OutExtensions) const override;
+		virtual bool GetThumbnailActionOverlay(const FAssetData& InAssetData, FAssetActionThumbnailOverlayInfo& OutActionOverlayInfo) const override;
+		virtual EAssetCommandResult ActivateAssets(const FAssetActivateArgs& ActivateArgs) const override;
+		virtual void GetAssetActionButtonExtensions(const FAssetData& InAssetData, TArray<FAssetButtonActionExtension>& OutExtensions) const override;
 
- 	virtual const FSlateBrush* GetThumbnailBrush(const FAssetData& InAssetData, const FName InClassName) const override;
- 	virtual const FSlateBrush* GetIconBrush(const FAssetData& InAssetData, const FName InClassName) const override;
-	// UAssetDefinition End
-};
\ No newline at end of file
+ 		virtual const FSlateBrush* GetThumbnailBrush(const FAssetData& InAssetData, const FName InClassName) const override;
+ 		virtual const FSlateBrush* GetIconBrush(const FAssetData& InAssetData, const FName InClassName) const override;
+		// UAssetDefinition End
+	};
+} // namespace UE::Subsonic
\ No newline at end of file
diff --git a/Engine/Plugins/Experimental/Subsonic/Source/SubsonicEditor/Private/SubsonicEventCollectionEditor.cpp b/Engine/Plugins/Experimental/Subsonic/Source/SubsonicEditor/Private/SubsonicEventCollectionEditor.cpp
--- a/Engine/Plugins/Experimental/Subsonic/Source/SubsonicEditor/Private/SubsonicEventCollectionEditor.cpp
+++ b/Engine/Plugins/Experimental/Subsonic/Source/SubsonicEditor/Private/SubsonicEventCollectionEditor.cpp
@@ -61,7 +61,7 @@ namespace UE::Subsonic::Editor
 		}
 	}
 
-	void FEventCollectionEditor::Init(const EToolkitMode::Type Mode, const TSharedPtr<IToolkitHost>& InitToolkitHost, USubsonicEventCollection& InEventCollection)
+	void FEventCollectionEditor::Init(const EToolkitMode::Type Mode, const TSharedPtr<IToolkitHost>& InitToolkitHost, UE::Subsonic::USubsonicEventCollection& InEventCollection)
 	{
 		// Support undo/redo
 		InEventCollection.SetFlags(RF_Transactional);
@@ -125,7 +125,7 @@ namespace UE::Subsonic::Editor
 		RegenerateMenusAndToolbars();
 	}
 
-	USubsonicEventCollection* FEventCollectionEditor::GetEventCollection() const
+	UE::Subsonic::USubsonicEventCollection* FEventCollectionEditor::GetEventCollection() const
 	{
 		return Cast<USubsonicEventCollection>(GetEditingObject());
 	}
@@ -194,5 +194,5 @@ namespace UE::Subsonic::Editor
 		FSlateApplication::Get().DismissAllMenus();
 		FEditorUndoClient::PostUndo(bSuccess);
 	}
-}
+} // namespace UE::Subsonic::Editor
 #undef LOCTEXT_NAMESPACE // SubsonicEditor
diff --git a/Engine/Plugins/Experimental/Subsonic/Source/SubsonicEditor/Private/SubsonicEventCollectionEditor.h b/Engine/Plugins/Experimental/Subsonic/Source/SubsonicEditor/Private/SubsonicEventCollectionEditor.h
--- a/Engine/Plugins/Experimental/Subsonic/Source/SubsonicEditor/Private/SubsonicEventCollectionEditor.h
+++ b/Engine/Plugins/Experimental/Subsonic/Source/SubsonicEditor/Private/SubsonicEventCollectionEditor.h
@@ -12,44 +12,46 @@
 #include "Toolkits/IToolkitHost.h"
 
 
-// Forward Declarations
-class USubsonicEventCollection;
-
-
-namespace UE::Subsonic::Editor
+namespace UE::Subsonic
 {
-	class FEventCollectionEditor : public FAssetEditorToolkit, public FNotifyHook, public FEditorUndoClient
-	{
-	public:
-		static const FName EditorName;
-
-		FEventCollectionEditor() = default;
-		virtual ~FEventCollectionEditor();
-
-		virtual void RegisterTabSpawners(const TSharedRef<FTabManager>& TabManager) override;
-		virtual void UnregisterTabSpawners(const TSharedRef<FTabManager>& TabManager) override;
-
-		void Init(const EToolkitMode::Type Mode, const TSharedPtr<IToolkitHost>& InitToolkitHost, USubsonicEventCollection& EventCollection);
+	// Forward Declarations
+	class USubsonicEventCollection;
 
-		USubsonicEventCollection* GetEventCollection() const;
-
-		/** IToolkit interface */
-		virtual FName GetToolkitFName() const override;
-		virtual FText GetBaseToolkitName() const override;
-		virtual FString GetWorldCentricTabPrefix() const override;
-		virtual FLinearColor GetWorldCentricTabColorScale() const override;
-		virtual const FSlateBrush* GetDefaultTabIcon() const override;
-		virtual FLinearColor GetDefaultTabColor() const override;
-
-		/** IAssetEditorInstance interface */
-		virtual FName GetEditorName() const override;
-
-		/** FEditorUndoClient Interface */
-		virtual void PostUndo(bool bSuccess) override;
-		virtual void PostRedo(bool bSuccess) override { PostUndo(bSuccess); }
-
-	private:
-		TSharedPtr<IDetailsView> DetailsView;
-		TSharedPtr<FTabManager::FStack> DetailsStack;
-	};
-} // namespace UE::Subsonic::Editor
+	namespace Editor
+	{
+		class FEventCollectionEditor : public FAssetEditorToolkit, public FNotifyHook, public FEditorUndoClient
+		{
+		public:
+			static const FName EditorName;
+
+			FEventCollectionEditor() = default;
+			virtual ~FEventCollectionEditor();
+
+			virtual void RegisterTabSpawners(const TSharedRef<FTabManager>& TabManager) override;
+			virtual void UnregisterTabSpawners(const TSharedRef<FTabManager>& TabManager) override;
+
+			void Init(const EToolkitMode::Type Mode, const TSharedPtr<IToolkitHost>& InitToolkitHost, USubsonicEventCollection& EventCollection);
+
+			USubsonicEventCollection* GetEventCollection() const;
+
+			/** IToolkit interface */
+			virtual FName GetToolkitFName() const override;
+			virtual FText GetBaseToolkitName() const override;
+			virtual FString GetWorldCentricTabPrefix() const override;
+			virtual FLinearColor GetWorldCentricTabColorScale() const override;
+			virtual const FSlateBrush* GetDefaultTabIcon() const override;
+			virtual FLinearColor GetDefaultTabColor() const override;
+
+			/** IAssetEditorInstance interface */
+			virtual FName GetEditorName() const override;
+
+			/** FEditorUndoClient Interface */
+			virtual void PostUndo(bool bSuccess) override;
+			virtual void PostRedo(bool bSuccess) override { PostUndo(bSuccess); }
+
+		private:
+			TSharedPtr<IDetailsView> DetailsView;
+			TSharedPtr<FTabManager::FStack> DetailsStack;
+		};
+	} // namespace Editor
+} // namespace UE::Subsonic
\ No newline at end of file
diff --git a/Engine/Plugins/Experimental/Subsonic/Source/SubsonicEditor/Private/SubsonicEventCollectionFactory.cpp b/Engine/Plugins/Experimental/Subsonic/Source/SubsonicEditor/Private/SubsonicEventCollectionFactory.cpp
--- a/Engine/Plugins/Experimental/Subsonic/Source/SubsonicEditor/Private/SubsonicEventCollectionFactory.cpp
+++ b/Engine/Plugins/Experimental/Subsonic/Source/SubsonicEditor/Private/SubsonicEventCollectionFactory.cpp
@@ -6,17 +6,20 @@
 #include UE_INLINE_GENERATED_CPP_BY_NAME(SubsonicEventCollectionFactory)
 
 
-USubsonicEventCollectionFactory::USubsonicEventCollectionFactory(const FObjectInitializer& ObjectInitializer)
-	: Super(ObjectInitializer)
+namespace UE::Subsonic::Editor
 {
-	bCreateNew = true;
-	bEditorImport = false;
-	bEditAfterNew = true;
+	USubsonicEventCollectionFactory::USubsonicEventCollectionFactory(const FObjectInitializer& ObjectInitializer)
+		: Super(ObjectInitializer)
+	{
+		bCreateNew = true;
+		bEditorImport = false;
+		bEditAfterNew = true;
 
-	SupportedClass = USubsonicEventCollection::StaticClass();
-}
+		SupportedClass = USubsonicEventCollection::StaticClass();
+	}
 
-UObject* USubsonicEventCollectionFactory::FactoryCreateNew(UClass* InClass, UObject* InParent, FName InName, EObjectFlags InFlags, UObject* InContext, FFeedbackContext* InFeedbackContext)
-{
-	return NewObject<USubsonicEventCollection>(InParent, InClass, InName, InFlags);
-}
+	UObject* USubsonicEventCollectionFactory::FactoryCreateNew(UClass* InClass, UObject* InParent, FName InName, EObjectFlags InFlags, UObject* InContext, FFeedbackContext* InFeedbackContext)
+	{
+		return NewObject<USubsonicEventCollection>(InParent, InClass, InName, InFlags);
+	}
+} // namespace UE::Subsonic::Editor
diff --git a/Engine/Plugins/Experimental/Subsonic/Source/SubsonicEditor/Private/SubsonicEventCollectionFactory.h b/Engine/Plugins/Experimental/Subsonic/Source/SubsonicEditor/Private/SubsonicEventCollectionFactory.h
--- a/Engine/Plugins/Experimental/Subsonic/Source/SubsonicEditor/Private/SubsonicEventCollectionFactory.h
+++ b/Engine/Plugins/Experimental/Subsonic/Source/SubsonicEditor/Private/SubsonicEventCollectionFactory.h
@@ -8,12 +8,15 @@
 #include "SubsonicEventCollectionFactory.generated.h"
 
 
-UCLASS()
-class USubsonicEventCollectionFactory : public UFactory
+namespace UE::Subsonic::Editor
 {
-	GENERATED_UCLASS_BODY()
+	UCLASS()
+		class USubsonicEventCollectionFactory : public UFactory
+	{
+		GENERATED_UCLASS_BODY()
 
-	//~ Begin UFactory Interface
-	virtual UObject* FactoryCreateNew(UClass* InClass, UObject* InParent, FName InName, EObjectFlags InFlags, UObject* InContext, FFeedbackContext* InFeedbackContext) override;
-	//~ Begin UFactory Interface
-};
+		//~ Begin UFactory Interface
+		virtual UObject* FactoryCreateNew(UClass* InClass, UObject* InParent, FName InName, EObjectFlags InFlags, UObject* InContext, FFeedbackContext* InFeedbackContext) override;
+		//~ Begin UFactory Interface
+	};
+} // namespace UE::Subsonic::Editor
diff --git a/Engine/Plugins/Experimental/Subsonic/Source/SubsonicEditor/SubsonicEditor.Build.cs b/Engine/Plugins/Experimental/Subsonic/Source/SubsonicEditor/SubsonicEditor.Build.cs
--- a/Engine/Plugins/Experimental/Subsonic/Source/SubsonicEditor/SubsonicEditor.Build.cs
+++ b/Engine/Plugins/Experimental/Subsonic/Source/SubsonicEditor/SubsonicEditor.Build.cs
@@ -30,5 +30,7 @@ public class SubsonicEditor : ModuleRules
 				"UnrealEd"
 			}
 		);
+
+		bAllowUETypesInNamespaces = true;
 	}
 }
