Remove "pointless" virtual on dtor
#rnx

#Synced-CL 50733540

diff --git a/Engine/Plugins/Experimental/Subsonic/Source/Subsonic/Public/SubsonicExecutor.h b/Engine/Plugins/Experimental/Subsonic/Source/Subsonic/Public/SubsonicExecutor.h
--- a/Engine/Plugins/Experimental/Subsonic/Source/Subsonic/Public/SubsonicExecutor.h
+++ b/Engine/Plugins/Experimental/Subsonic/Source/Subsonic/Public/SubsonicExecutor.h
@@ -44,7 +44,7 @@ namespace UE::Subsonic
 	public:
 		FSubsonicExecutor(FSubsonicExecutor&&);
 		FSubsonicExecutor& operator=(const FSubsonicExecutor& InOther) = delete;
-		virtual ~FSubsonicExecutor();
+		~FSubsonicExecutor();
 
 		UE_API FSubsonicExecutor& operator=(FSubsonicExecutor&& InOther);
 
