Add Subsonic Collection PropertyBag for providing shared, user-defined state across actions

#rnx

diff --git a/Engine/Plugins/Experimental/Subsonic/Source/Subsonic/Private/SubsonicEventCollection.cpp b/Engine/Plugins/Experimental/Subsonic/Source/Subsonic/Private/SubsonicEventCollection.cpp
--- a/Engine/Plugins/Experimental/Subsonic/Source/Subsonic/Private/SubsonicEventCollection.cpp
+++ b/Engine/Plugins/Experimental/Subsonic/Source/Subsonic/Private/SubsonicEventCollection.cpp
@@ -45,19 +45,34 @@ namespace UE::Subsonic
 	FSubsonicEventCollectionDefinition::FSubsonicEventCollectionDefinition(FSubsonicEventCollectionDefinition&& InOther)
 	{
 		Unregister();
+
 		CollectionId = InOther.CollectionId;
 		InOther.CollectionId = GetInvalidId();
 
 		Events = MoveTemp(InOther.Events);
+		InOther.Events.Reset();
+
+#if WITH_EDITORONLY_DATA
+		Parameters = MoveTemp(InOther.Parameters);
+		InOther.Parameters.Reset();
+#endif // WITH_EDITORONLY_DATA
 	}
 
 	FSubsonicEventCollectionDefinition& FSubsonicEventCollectionDefinition::operator =(FSubsonicEventCollectionDefinition&& InOther)
 	{
 		Unregister();
+
 		CollectionId = InOther.CollectionId;
 		InOther.CollectionId = GetInvalidId();
 
 		Events = MoveTemp(InOther.Events);
+		InOther.Events.Reset();
+
+#if WITH_EDITORONLY_DATA
+		Parameters = MoveTemp(InOther.Parameters);
+		InOther.Parameters.Reset();
+#endif // WITH_EDITORONLY_DATA
+
 		return *this;
 	}
 
@@ -111,6 +126,26 @@ namespace UE::Subsonic
 		return UE::Subsonic::CollectionPrivate::InvalidId;
 	}
 
+#if WITH_EDITOR
+	FName FSubsonicEventCollectionDefinition::GetEventsPropertyName()
+	{
+#if WITH_EDITORONLY_DATA
+		return GET_MEMBER_NAME_CHECKED(FSubsonicEventCollectionDefinition, Events);
+#else
+		return { };
+#endif // !WITH_EDITORONLY_DATA
+	}
+
+	FName FSubsonicEventCollectionDefinition::GetParametersPropertyName()
+	{
+#if WITH_EDITORONLY_DATA
+		return GET_MEMBER_NAME_CHECKED(FSubsonicEventCollectionDefinition, Parameters);
+#else
+		return { };
+#endif // !WITH_EDITORONLY_DATA
+	}
+#endif // WITH_EDITOR
+
 	void FSubsonicEventCollectionDefinition::Register(const FCollectionHandle& InHandle)
 	{
 		if (GEngine && IsValid())
@@ -135,6 +170,43 @@ namespace UE::Subsonic
 		CollectionId = CollectionPrivate::InvalidId;
 	}
 
+	const FSubsonicEventCollectionDefinition& USubsonicEventCollection::GetDefinition() const
+	{
+		return CollectionDefinition;
+	}
+
+#if WITH_EDITOR
+	FName USubsonicEventCollection::GetDefinitionPropertyName()
+	{
+		return GET_MEMBER_NAME_CHECKED(USubsonicEventCollection, CollectionDefinition);
+	}
+#endif // WITH_EDITOR
+
+	FCollectionHandle USubsonicEventCollection::GetHandle() const
+	{
+#if WITH_EDITORONLY_DATA
+		FName CollectionName;
+		if (const UPackage* CollectionPackage = GetPackage())
+		{
+			const FTopLevelAssetPath AssetPath(CollectionPackage);
+			CollectionName = AssetPath.GetPackageName();
+		}
+#endif // WITH_EDITORONLY_DATA
+
+		return FCollectionHandle
+		{
+#if WITH_EDITORONLY_DATA
+			.CollectionName = CollectionName,
+#endif // WITH_EDITORONLY_DATA
+			.CollectionId = CollectionDefinition.GetCollectionId(),
+		};
+	}
+
+	void USubsonicEventCollection::SetDefinition(FSubsonicEventCollectionDefinition&& Definition)
+	{
+		CollectionDefinition = MoveTemp(Definition);
+	}
+
 	void USubsonicEventCollection::PostInitProperties()
 	{
 		Super::PostInitProperties();
diff --git a/Engine/Plugins/Experimental/Subsonic/Source/Subsonic/Public/SubsonicEventCollection.h b/Engine/Plugins/Experimental/Subsonic/Source/Subsonic/Public/SubsonicEventCollection.h
--- a/Engine/Plugins/Experimental/Subsonic/Source/Subsonic/Public/SubsonicEventCollection.h
+++ b/Engine/Plugins/Experimental/Subsonic/Source/Subsonic/Public/SubsonicEventCollection.h
@@ -4,7 +4,10 @@
 #include "AudioDefines.h"
 #include "GameplayTagContainer.h"
 #include "NativeGameplayTags.h"
+#include "PropertyBindingBindingCollectionOwner.h"
+#include "PropertyBindingPath.h"
 #include "StructUtils/InstancedStruct.h"
+#include "StructUtils/PropertyBag.h"
 #include "SubsonicExecutor.h"
 #include "UObject/NameTypes.h"
 
@@ -83,6 +86,11 @@ namespace UE::Subsonic
 		// Returns Id of invalid (default constructed or unregistered/reset collection).
 		UE_API static uint32 GetInvalidId();
 
+#if WITH_EDITOR
+		UE_API static FName GetEventsPropertyName();
+		UE_API static FName GetParametersPropertyName();
+#endif // WITH_EDITOR
+
 	private:
 		// Register's the given collection with the Subsonic Subsystem.
 		UE_INTERNAL void Register(const FCollectionHandle& InHandle);
@@ -98,12 +106,20 @@ namespace UE::Subsonic
 		// having to utilize UObject wrapper.
 		uint32 CollectionId = INDEX_NONE;
 
-		UPROPERTY(EditAnywhere, BlueprintReadOnly, Category = "Subsonic", meta = (AllowPrivateAccess))
+		UPROPERTY(EditAnywhere, BlueprintReadOnly, Category = "Events", meta = (AllowPrivateAccess))
 		TMap<FGameplayTag, FSubsonicEvent> Events;
 
-		friend class USubsonicEventCollection;
+#if WITH_EDITORONLY_DATA
+		UPROPERTY(EditAnywhere, Category = "Parameters")
+		FInstancedPropertyBag Parameters;
+#endif // WITH_EDITORONLY_DATA
+
+// 		UPROPERTY(VisibleAnywhere, Category = "Subsonic", AdvancedDisplay)
+// 		FSubsonicParameterBindingCollection Bindings;
 
 		void AssignId();
+
+		friend class USubsonicEventCollection;
 	};
 
 	UCLASS(MinimalAPI, BlueprintType)
@@ -112,36 +128,14 @@ namespace UE::Subsonic
 		GENERATED_BODY()
 
 	public:
-		const FSubsonicEventCollectionDefinition& GetDefinition() const
-		{
-			return CollectionDefinition;
-		}
-
-		FCollectionHandle GetHandle() const
-		{
-	#if WITH_EDITORONLY_DATA
-			FName CollectionName;
-			if (const UPackage* CollectionPackage = GetPackage())
-			{
-				const FTopLevelAssetPath AssetPath(CollectionPackage);
-				CollectionName = AssetPath.GetPackageName();
-			}
-	#endif // WITH_EDITORONLY_DATA
-
-			return FCollectionHandle
-			{
-	#if WITH_EDITORONLY_DATA
-				.CollectionName = CollectionName,
-	#endif // WITH_EDITORONLY_DATA
-				.CollectionId = CollectionDefinition.GetCollectionId(),
-			};
-		}
-
-		void SetDefinition(FSubsonicEventCollectionDefinition Definition)
-		{
-			Definition.AssignId();
-			CollectionDefinition = MoveTemp(Definition);
-		}
+		UE_API const FSubsonicEventCollectionDefinition& GetDefinition() const;
+
+#if WITH_EDITOR
+		UE_API static FName GetDefinitionPropertyName();
+#endif // WITH_EDITOR
+
+		UE_API FCollectionHandle GetHandle() const;
+		UE_API void SetDefinition(FSubsonicEventCollectionDefinition&& Definition);
 
 	private:
 		UPROPERTY(EditAnywhere, Category = "Subsonic", meta = (ShowOnlyInnerProperties))
diff --git a/Engine/Plugins/Experimental/Subsonic/Source/Subsonic/Subsonic.Build.cs b/Engine/Plugins/Experimental/Subsonic/Source/Subsonic/Subsonic.Build.cs
--- a/Engine/Plugins/Experimental/Subsonic/Source/Subsonic/Subsonic.Build.cs
+++ b/Engine/Plugins/Experimental/Subsonic/Source/Subsonic/Subsonic.Build.cs
@@ -13,7 +13,8 @@ public class Subsonic : ModuleRules
 			{
 				"AudioMixerCore",
 				"Core",
-				"GameplayTags"
+				"GameplayTags",
+				"PropertyBindingUtils"
 			}
 		);
 			
diff --git a/Engine/Plugins/Experimental/Subsonic/Source/SubsonicEditor/Private/SubsonicEditorModule.cpp b/Engine/Plugins/Experimental/Subsonic/Source/SubsonicEditor/Private/SubsonicEditorModule.cpp
--- a/Engine/Plugins/Experimental/Subsonic/Source/SubsonicEditor/Private/SubsonicEditorModule.cpp
+++ b/Engine/Plugins/Experimental/Subsonic/Source/SubsonicEditor/Private/SubsonicEditorModule.cpp
@@ -1,21 +1,66 @@
 // Copyright Epic Games, Inc. All Rights Reserved.
 
-#include "SubsonicSlateStyle.h"
 #include "Modules/ModuleManager.h"
+#include "PropertyEditorModule.h"
+#include "SubsonicEventCollectionViews.h"
+#include "SubsonicSlateStyle.h"
 
 
 namespace UE::Subsonic::Editor
 {
 	class FModule : public IModuleInterface
 	{
+		struct ICollectionViewEntryBase { virtual ~ICollectionViewEntryBase() = default; };
+		using FViewEntryPtr = TUniquePtr<ICollectionViewEntryBase>;
+
+		template <typename TViewClass, typename TCustomizationClass>
+		class FCollectionViewEntry : public ICollectionViewEntryBase
+		{
+			FName ClassName;
+
+			FCollectionViewEntry(FPropertyEditorModule& PropertyModule)
+			{
+				UClass* ViewClass = TViewClass::StaticClass();
+				check(ViewClass);
+				ClassName = ViewClass->GetFName();
+				PropertyModule.RegisterCustomClassLayout(
+					ClassName,
+					FOnGetDetailCustomizationInstance::CreateLambda([]() { return MakeShared<TCustomizationClass>(); }));
+			}
+
+		public:
+			~FCollectionViewEntry()
+			{
+				if (FPropertyEditorModule* PropertyModule = FModuleManager::GetModulePtr<FPropertyEditorModule>("PropertyEditor"))
+				{
+					PropertyModule->UnregisterCustomClassLayout(ClassName);
+				}
+			}
+
+			static FViewEntryPtr Register(FPropertyEditorModule& PropertyModule)
+			{
+				return FViewEntryPtr(new FCollectionViewEntry<TViewClass, TCustomizationClass>(PropertyModule));
+			}
+		};
+
+		TArray<FViewEntryPtr> CollectionViewEntries;
+
 	public:
 		virtual void StartupModule() override
 		{
+			// Register Custom Collection Class Views
+			{
+				FPropertyEditorModule& PropertyModule = FModuleManager::LoadModuleChecked<FPropertyEditorModule>("PropertyEditor");
+				CollectionViewEntries.Add(FCollectionViewEntry<USubsonicEventsView, FCollectionEventsDetailCustomization>::Register(PropertyModule));
+				CollectionViewEntries.Add(FCollectionViewEntry<USubsonicParametersView, FCollectionParametersDetailCustomization>::Register(PropertyModule));
+			}
+
 			StyleSet = MakeShared<FSlateStyle>();
 		}
 
 		virtual void ShutdownModule() override
 		{
+			CollectionViewEntries.Reset();
 			StyleSet.Reset();
 		}
 
@@ -23,5 +68,4 @@ namespace UE::Subsonic::Editor
 	};
 } // namespace UE::Subsonic::Editor
 
-
 IMPLEMENT_MODULE(UE::Subsonic::Editor::FModule, SubsonicEditor);
diff --git a/Engine/Plugins/Experimental/Subsonic/Source/SubsonicEditor/Private/SubsonicEventCollectionEditor.cpp b/Engine/Plugins/Experimental/Subsonic/Source/SubsonicEditor/Private/SubsonicEventCollectionEditor.cpp
--- a/Engine/Plugins/Experimental/Subsonic/Source/SubsonicEditor/Private/SubsonicEventCollectionEditor.cpp
+++ b/Engine/Plugins/Experimental/Subsonic/Source/SubsonicEditor/Private/SubsonicEventCollectionEditor.cpp
@@ -9,6 +9,7 @@
 #include "Widgets/Docking/SDockTab.h"
 #include "Styling/SlateStyleRegistry.h"
 #include "SubsonicEventCollection.h"
+#include "SubsonicEventCollectionViews.h"
 
 #define LOCTEXT_NAMESPACE "SubsonicEditor"
 
@@ -17,7 +18,8 @@ namespace UE::Subsonic::Editor
 {
 	namespace EventCollectionEditorPrivate
 	{
-		const FName TabName_Details = "SubsonicEventCollectionEditor_Details";
+		const FName TabName_Events = "SubsonicEventCollectionEditor_Events";
+		const FName TabName_Parameters = "SubsonicEventCollectionEditor_Parameters";
 	} // namespace EventCollectionEditorPrivate
 
 	const FName FEventCollectionEditor::EditorName = "SubsonicEventCollectionEditor";
@@ -31,17 +33,24 @@ namespace UE::Subsonic::Editor
 		Args.NotifyHook = this;
 
 		FPropertyEditorModule& PropertyModule = FModuleManager::LoadModuleChecked<FPropertyEditorModule>("PropertyEditor");
-		TSharedRef<IDetailsView> NewDetailsView = PropertyModule.CreateDetailView(Args);
 
-		InTabManager->RegisterTabSpawner(EventCollectionEditorPrivate::TabName_Details, FOnSpawnTab::CreateLambda([NewDetailsView](const FSpawnTabArgs& Args)
+		EventsDetailsView = PropertyModule.CreateDetailView(Args);
+		InTabManager->RegisterTabSpawner(EventCollectionEditorPrivate::TabName_Events, FOnSpawnTab::CreateLambda([ViewRef = EventsDetailsView.ToSharedRef()](const FSpawnTabArgs& Args)
 			{
-				return SNew(SDockTab).Label(LOCTEXT("SubsonicEventCollectionDetailsTitle", "Details"))[ NewDetailsView ];
+				return SNew(SDockTab).Label(LOCTEXT("SubsonicEventCollectionEventsTitle", "Events"))[ ViewRef ];
 			}))
-			.SetDisplayName(LOCTEXT("DetailsTab", "Details"))
+			.SetDisplayName(LOCTEXT("EventsTab", "Events"))
 			.SetGroup(WorkspaceMenuCategory.ToSharedRef())
 			.SetIcon(FSlateIcon(FAppStyle::GetAppStyleSetName(), "LevelEditor.Tabs.Details"));
 
-		DetailsView = NewDetailsView;
+		ParametersDetailsView = PropertyModule.CreateDetailView(Args);
+		InTabManager->RegisterTabSpawner(EventCollectionEditorPrivate::TabName_Parameters, FOnSpawnTab::CreateLambda([ViewRef = ParametersDetailsView.ToSharedRef()](const FSpawnTabArgs& Args)
+			{
+				return SNew(SDockTab).Label(LOCTEXT("SubsonicEventCollectionParametersTitle", "Parameters"))[ ViewRef ];
+			}))
+			.SetDisplayName(LOCTEXT("ParametersTab", "Parameters"))
+			.SetGroup(WorkspaceMenuCategory.ToSharedRef())
+			.SetIcon(FSlateIcon(FAppStyle::GetAppStyleSetName(), "LevelEditor.Tabs.Details"));
 
 		FAssetEditorToolkit::RegisterTabSpawners(InTabManager);
 	}
@@ -50,7 +59,8 @@ namespace UE::Subsonic::Editor
 	{
 		FAssetEditorToolkit::UnregisterTabSpawners(InTabManager);
 
-		InTabManager->UnregisterTabSpawner(EventCollectionEditorPrivate::TabName_Details);
+		InTabManager->UnregisterTabSpawner(EventCollectionEditorPrivate::TabName_Events);
+		InTabManager->UnregisterTabSpawner(EventCollectionEditorPrivate::TabName_Parameters);
 	}
 
 	FEventCollectionEditor::~FEventCollectionEditor()
@@ -61,7 +71,7 @@ namespace UE::Subsonic::Editor
 		}
 	}
 
-	void FEventCollectionEditor::Init(const EToolkitMode::Type Mode, const TSharedPtr<IToolkitHost>& InitToolkitHost, UE::Subsonic::USubsonicEventCollection& InEventCollection)
+	void FEventCollectionEditor::Init(const EToolkitMode::Type Mode, const TSharedPtr<IToolkitHost>& InitToolkitHost, USubsonicEventCollection& InEventCollection)
 	{
 		// Support undo/redo
 		InEventCollection.SetFlags(RF_Transactional);
@@ -71,31 +81,28 @@ namespace UE::Subsonic::Editor
 			GEditor->RegisterForUndo(this);
 		}
 
-		DetailsStack = FTabManager::NewStack()
-			->SetSizeCoefficient(0.50f)
-			->SetHideTabWell(false)
-			->AddTab(EventCollectionEditorPrivate::TabName_Details, ETabState::OpenedTab);
+		ParametersStack = FTabManager::NewStack()
+			->SetSizeCoefficient(0.33f)
+			->SetHideTabWell(true)
+			->AddTab(EventCollectionEditorPrivate::TabName_Parameters, ETabState::OpenedTab);
+
+		EventsStack = FTabManager::NewStack()
+			->SetSizeCoefficient(0.67f)
+			->SetHideTabWell(true)
+			->AddTab(EventCollectionEditorPrivate::TabName_Events, ETabState::OpenedTab);
 
-		const TSharedRef<FTabManager::FLayout> StandaloneDefaultLayout = FTabManager::NewLayout("Standalone_SubsonicEventCollectionEditor_Layout_v1")
+		const TSharedRef<FTabManager::FLayout> StandaloneDefaultLayout = FTabManager::NewLayout("Standalone_SubsonicEventCollectionEditor_Layout_v2")
 			->AddArea
 			(
 				FTabManager::NewPrimaryArea()
-				->SetOrientation(Orient_Vertical)
-				->Split(FTabManager::NewSplitter()
-					->SetOrientation(Orient_Horizontal)
-					->Split
-					(
-						FTabManager::NewSplitter()
-						->SetSizeCoefficient(0.15f)
-						->SetOrientation(Orient_Vertical)
-						->Split
-						(
-							FTabManager::NewStack()
-							->SetSizeCoefficient(0.25f)
-							->SetHideTabWell(false)
-							->AddTab(EventCollectionEditorPrivate::TabName_Details, ETabState::OpenedTab)
-						)
-					)
+				->SetOrientation(Orient_Horizontal)
+				->Split
+				(
+					ParametersStack.ToSharedRef()
+				)
+				->Split
+				(
+					EventsStack.ToSharedRef()
 				)
 			);
 
@@ -117,15 +124,25 @@ namespace UE::Subsonic::Editor
 				bUseSmallToolbarIcons);
 		}
 
-		if (DetailsView.IsValid())
+		USubsonicEventCollection* EventCollection = CastChecked<USubsonicEventCollection>(GetEditingObject());
+		if (EventsDetailsView.IsValid())
+		{
+			EventsView.Reset(NewObject<USubsonicEventsView>());
+			EventsView->SetCollection(EventCollection);
+			EventsDetailsView->SetObject(EventsView.Get());
+		}
+
+		if (ParametersDetailsView.IsValid())
 		{
-			DetailsView->SetObject(&InEventCollection);
+			ParametersView.Reset(NewObject<USubsonicParametersView>());
+			ParametersView->SetCollection(EventCollection);
+			ParametersDetailsView->SetObject(ParametersView.Get());
 		}
 
 		RegenerateMenusAndToolbars();
 	}
 
-	UE::Subsonic::USubsonicEventCollection* FEventCollectionEditor::GetEventCollection() const
+	USubsonicEventCollection* FEventCollectionEditor::GetEventCollection() const
 	{
 		return Cast<USubsonicEventCollection>(GetEditingObject());
 	}
diff --git a/Engine/Plugins/Experimental/Subsonic/Source/SubsonicEditor/Private/SubsonicEventCollectionEditor.h b/Engine/Plugins/Experimental/Subsonic/Source/SubsonicEditor/Private/SubsonicEventCollectionEditor.h
--- a/Engine/Plugins/Experimental/Subsonic/Source/SubsonicEditor/Private/SubsonicEventCollectionEditor.h
+++ b/Engine/Plugins/Experimental/Subsonic/Source/SubsonicEditor/Private/SubsonicEventCollectionEditor.h
@@ -7,10 +7,11 @@
 #include "Framework/Docking/TabManager.h"
 #include "IDetailsView.h"
 #include "Misc/NotifyHook.h"
+#include "SubsonicEventCollectionViews.h"
 #include "Templates/SharedPointer.h"
 #include "Toolkits/AssetEditorToolkit.h"
 #include "Toolkits/IToolkitHost.h"
-
+#include "UObject/StrongObjectPtrTemplates.h"
 
 namespace UE::Subsonic
 {
@@ -50,8 +51,13 @@ namespace UE::Subsonic
 			virtual void PostRedo(bool bSuccess) override { PostUndo(bSuccess); }
 
 		private:
-			TSharedPtr<IDetailsView> DetailsView;
-			TSharedPtr<FTabManager::FStack> DetailsStack;
+			TSharedPtr<IDetailsView> EventsDetailsView;
+			TSharedPtr<FTabManager::FStack> EventsStack;
+			TStrongObjectPtr<USubsonicEventsView> EventsView;
+
+			TSharedPtr<IDetailsView> ParametersDetailsView;
+			TSharedPtr<FTabManager::FStack> ParametersStack;
+			TStrongObjectPtr<USubsonicParametersView> ParametersView;
 		};
 	} // namespace Editor
 } // namespace UE::Subsonic
\ No newline at end of file
diff --git a/Engine/Plugins/Experimental/Subsonic/Source/SubsonicEditor/Private/SubsonicEventCollectionViews.cpp b/Engine/Plugins/Experimental/Subsonic/Source/SubsonicEditor/Private/SubsonicEventCollectionViews.cpp
new file mode 100644
--- /dev/null
+++ b/Engine/Plugins/Experimental/Subsonic/Source/SubsonicEditor/Private/SubsonicEventCollectionViews.cpp
@@ -0,0 +1,59 @@
+// Copyright Epic Games, Inc. All Rights Reserved.
+
+#include "SubsonicEventCollectionViews.h"
+
+#include "DetailCategoryBuilder.h"
+#include "DetailWidgetRow.h"
+#include "IDetailChildrenBuilder.h"
+#include "Internationalization/Internationalization.h"
+#include "Misc/AssertionMacros.h"
+#include "SubsonicEventCollection.h"
+
+#include UE_INLINE_GENERATED_CPP_BY_NAME(SubsonicEventCollectionViews)
+
+#define LOCTEXT_NAMESPACE "SubsonicEditor"
+
+
+namespace UE::Subsonic::Editor
+{
+	void FCollectionPropertyCustomizationBase::CustomizeDetails(IDetailLayoutBuilder& DetailLayout)
+	{
+		// Only supports modifying a single collection at a time
+		TArray<TWeakObjectPtr<UObject>> Objects;
+		DetailLayout.GetObjectsBeingCustomized(Objects);
+		if (Objects.Num() == 1 && Objects.Last().IsValid())
+		{
+			if (USubsonicEventCollectionViewBase* View = CastChecked<USubsonicEventCollectionViewBase>(Objects.Last()))
+			{
+				TSharedPtr<IPropertyHandle> DefinitionHandle = DetailLayout.AddObjectPropertyData(
+					TArray<UObject*> { View->GetCollection() },
+					USubsonicEventCollection::GetDefinitionPropertyName());
+
+				if (DefinitionHandle.IsValid())
+				{
+					const FText CategoryDisplayName = HeaderName.IsEmptyOrWhitespace() ? FText::FromName(PropertyName) : HeaderName;
+					TSharedPtr<IPropertyHandle> ChildProperty = DefinitionHandle->GetChildHandle(PropertyName);
+					if (ChildProperty.IsValid())
+					{
+						IDetailCategoryBuilder& Category = DetailLayout.EditCategory(PropertyName, CategoryDisplayName);
+						Category.InitiallyCollapsed(false);
+						Category.AddProperty(ChildProperty);
+					}
+				}
+			}
+		}
+	}
+
+	FCollectionParametersDetailCustomization::FCollectionParametersDetailCustomization()
+	{
+		PropertyName = FSubsonicEventCollectionDefinition::GetParametersPropertyName();
+		HeaderName = LOCTEXT("ParametersDetails_DisplayName", "Parameters");
+	}
+
+	FCollectionEventsDetailCustomization::FCollectionEventsDetailCustomization()
+	{
+		PropertyName = FSubsonicEventCollectionDefinition::GetEventsPropertyName();
+		HeaderName = LOCTEXT("EventsDetails_DisplayName", "Events");
+	}
+} // namespace UE::Subsonic::Editor
+#undef LOCTEXT_NAMESPACE // "SubsonicEditor"
diff --git a/Engine/Plugins/Experimental/Subsonic/Source/SubsonicEditor/Private/SubsonicEventCollectionViews.h b/Engine/Plugins/Experimental/Subsonic/Source/SubsonicEditor/Private/SubsonicEventCollectionViews.h
new file mode 100644
--- /dev/null
+++ b/Engine/Plugins/Experimental/Subsonic/Source/SubsonicEditor/Private/SubsonicEventCollectionViews.h
@@ -0,0 +1,64 @@
+// Copyright Epic Games, Inc. All Rights Reserved.
+#pragma once
+
+#include "DetailLayoutBuilder.h"
+#include "IDetailCustomization.h"
+#include "SubsonicEventCollection.h"
+
+#include "SubsonicEventCollectionViews.generated.h"
+
+
+namespace UE::Subsonic::Editor
+{
+	// Base implementation for Subsonic Collection view wrappers that coincide with edit customizations
+	UCLASS(MinimalAPI)
+	class USubsonicEventCollectionViewBase : public UObject
+	{
+		GENERATED_BODY()
+
+	public:
+
+		void SetCollection(USubsonicEventCollection* InCollection) { Collection = InCollection; }
+		const USubsonicEventCollection* GetConstCollection() const { return Collection.Get(); }
+		USubsonicEventCollection* GetCollection() { return Collection.Get(); }
+
+	private:
+		TWeakObjectPtr<USubsonicEventCollection> Collection;
+	};
+
+	UCLASS(MinimalAPI)
+	class USubsonicParametersView : public USubsonicEventCollectionViewBase
+	{
+		GENERATED_BODY()
+	};
+
+	UCLASS(MinimalAPI)
+	class USubsonicEventsView : public USubsonicEventCollectionViewBase
+	{
+		GENERATED_BODY()
+	};
+
+	class FCollectionPropertyCustomizationBase : public IDetailCustomization
+	{
+	protected:
+		FName PropertyName;
+		FText HeaderName;
+
+	public:
+		// IDetailCustomization interface
+		virtual void CustomizeDetails(IDetailLayoutBuilder& DetailLayout) override;
+		// End of IDetailCustomization interface
+	};
+
+	class FCollectionParametersDetailCustomization : public FCollectionPropertyCustomizationBase
+	{
+	public:
+		FCollectionParametersDetailCustomization();
+	};
+
+	class FCollectionEventsDetailCustomization : public FCollectionPropertyCustomizationBase
+	{
+	public:
+		FCollectionEventsDetailCustomization();
+	};
+} // namespace UE::Subsonic::Editor
\ No newline at end of file
diff --git a/Engine/Plugins/Experimental/Subsonic/Subsonic.uplugin b/Engine/Plugins/Experimental/Subsonic/Subsonic.uplugin
--- a/Engine/Plugins/Experimental/Subsonic/Subsonic.uplugin
+++ b/Engine/Plugins/Experimental/Subsonic/Subsonic.uplugin
@@ -31,6 +31,10 @@
 		{
 			"Name": "AudioWidgets",
 			"Enabled": true
+		},
+		{
+			"Name": "PropertyBindingUtils",
+			"Enabled": true
 		}
 	]
 }
