trivial: Add exports to ToString() handle functions
#rnx

diff --git a/Engine/Plugins/Experimental/Subsonic/Source/Subsonic/Public/SubsonicHandles.h b/Engine/Plugins/Experimental/Subsonic/Source/Subsonic/Public/SubsonicHandles.h
--- a/Engine/Plugins/Experimental/Subsonic/Source/Subsonic/Public/SubsonicHandles.h
+++ b/Engine/Plugins/Experimental/Subsonic/Source/Subsonic/Public/SubsonicHandles.h
@@ -40,7 +40,7 @@ namespace UE::Subsonic
 			return A.CollectionId == B.CollectionId;
 		}
 
-		FString ToString() const;
+		UE_API FString ToString() const;
 	};
 
 	struct FEventHandle
@@ -61,7 +61,7 @@ namespace UE::Subsonic
 			return A.Collection == B.Collection && A.EventName == B.EventName;
 		}
 
-		FString ToString() const;
+		UE_API FString ToString() const;
 	};
 
 	struct FActionHandle
@@ -82,7 +82,7 @@ namespace UE::Subsonic
 			return A.Event == B.Event && A.Index == B.Index;
 		}
 
-		FString ToString() const;
+		UE_API FString ToString() const;
 	};
 } // namespace UE::Subsonic
 
