Subsonic Updates
- Transport handling of preview executor
- Private/Public Event field/UX
- More boilerplate for Subsonic Collection PropertyBag Binding support
- Stub icons for transport usage
- Stub action for executing other events within a collection (TODO: implement delay tracking in subscriber as example)
- New customized Subsonic EventTree view in Subsonic Event Collection Editor
- Cache subsonic action structs in USubsonicEditorSubsystem for faster iteration/access (potentially useful for future editor scripts as well) on editor module load.
- Add Icon prototype

#rnx

#Synced-CL 51062589

diff --git a/Engine/Plugins/Experimental/Subsonic/Source/Subsonic/Private/CoreActions/SubsonicAction_AudioComponent.cpp b/Engine/Plugins/Experimental/Subsonic/Source/Subsonic/Private/CoreActions/SubsonicAction_AudioComponent.cpp
--- a/Engine/Plugins/Experimental/Subsonic/Source/Subsonic/Private/CoreActions/SubsonicAction_AudioComponent.cpp
+++ b/Engine/Plugins/Experimental/Subsonic/Source/Subsonic/Private/CoreActions/SubsonicAction_AudioComponent.cpp
@@ -2,9 +2,12 @@
 
 #include "CoreActions/SubsonicAction_AudioComponent.h"
 #include "IAudioModulation.h"
+#include "Internationalization/Internationalization.h"
 #include "SubsonicHandles.h"
 
 
+#define LOCTEXT_NAMESPACE "SubsonicActions"
+
 namespace UE::Subsonic
 {
 	namespace AudioComponentPrivate
@@ -12,31 +15,31 @@ namespace UE::Subsonic
 		UAudioComponent* AddScopedComponent(const UE::Subsonic::FActionScopeKey& InKey, ESubsonicExecutionScope Scope, FName Name)
 		{
 			UAudioComponent* Component = nullptr;
-			if (USubsonicAudioComponentSubscriber* AudioComponentSubsystem = InKey.FindDeviceSubsystem<USubsonicAudioComponentSubscriber>())
+			if (USubsonicAudioComponentSubscriber* AudioComponentSubscriber = InKey.FindDeviceSubsystem<USubsonicAudioComponentSubscriber>())
 			{
 				switch (Scope)
 				{
 					case ESubsonicExecutionScope::Global:
 					{
-						Component = AudioComponentSubsystem->AddComponent(Name);
+						Component = AudioComponentSubscriber->AddComponent(Name);
 					}
 					break;
 
 					case ESubsonicExecutionScope::Collection:
 					{
-						Component = AudioComponentSubsystem->AddComponent(InKey.ScopeHandle.Event.Collection, Name);
+						Component = AudioComponentSubscriber->AddComponent(InKey.ScopeHandle.Event.Collection, Name);
 					}
 					break;
 
 					case ESubsonicExecutionScope::Event:
 					{
-						Component = AudioComponentSubsystem->AddComponent(InKey.ScopeHandle.Event, Name);
+						Component = AudioComponentSubscriber->AddComponent(InKey.ScopeHandle.Event, Name);
 					}
 					break;
 
 					case ESubsonicExecutionScope::Action:
 					{
-						Component = AudioComponentSubsystem->AddComponent(InKey.ScopeHandle, Name);
+						Component = AudioComponentSubscriber->AddComponent(InKey.ScopeHandle, Name);
 					}
 					break;
 				}
@@ -48,31 +51,31 @@ namespace UE::Subsonic
 		UAudioComponent* FindScopedComponent(const UE::Subsonic::FActionScopeKey& InKey, ESubsonicExecutionScope Scope, FName Name)
 		{
 			UAudioComponent* Component = nullptr;
-			if (USubsonicAudioComponentSubscriber* AudioComponentSubsystem = InKey.FindDeviceSubsystem<USubsonicAudioComponentSubscriber>())
+			if (USubsonicAudioComponentSubscriber* AudioComponentSubscriber = InKey.FindDeviceSubsystem<USubsonicAudioComponentSubscriber>())
 			{
 				switch (Scope)
 					{
 					case ESubsonicExecutionScope::Global:
 					{
-						Component = AudioComponentSubsystem->FindComponent(Name);
+						Component = AudioComponentSubscriber->FindComponent(Name);
 					}
 					break;
 
 					case ESubsonicExecutionScope::Collection:
 					{
-						Component = AudioComponentSubsystem->FindComponent(InKey.ScopeHandle.Event.Collection, Name);
+						Component = AudioComponentSubscriber->FindComponent(InKey.ScopeHandle.Event.Collection, Name);
 					}
 					break;
 
 					case ESubsonicExecutionScope::Event:
 					{
-						Component = AudioComponentSubsystem->FindComponent(InKey.ScopeHandle.Event, Name);
+						Component = AudioComponentSubscriber->FindComponent(InKey.ScopeHandle.Event, Name);
 					}
 					break;
 
 					case ESubsonicExecutionScope::Action:
 					{
-						Component = AudioComponentSubsystem->FindComponent(InKey.ScopeHandle, Name);
+						Component = AudioComponentSubscriber->FindComponent(InKey.ScopeHandle, Name);
 					}
 					break;
 				}
@@ -124,6 +127,56 @@ namespace UE::Subsonic
 		}
 	}
 
+#if WITH_EDITOR
+	FText FSubsonicEventAction_AudioComponentModify::GetDisplayInfo() const
+	{
+		FText NameText = FText::FromName(Name);
+		if (NameText.IsEmptyOrWhitespace())
+		{
+			NameText = Super::GetDisplayInfo();
+		}
+
+		if (!Modifiers.IsEmpty())
+		{
+			if (const UScriptStruct* ModifierStruct = Modifiers[0].GetScriptStruct())
+			{
+				if (ModifierStruct&& Modifiers.Num() > 1)
+				{
+					return FText::Format(LOCTEXT("AudioComponentModifierListDisplaySinglePlusInfo_Format", "{0} ('{1}' and {2} more)"),
+						NameText,
+						ModifierStruct->GetDisplayNameText(),
+						FText::AsNumber(Modifiers.Num() - 1));
+				}
+
+				return FText::Format(LOCTEXT("AudioComponentModifierListDisplayInfoSingle_Format", "{0} ({1})"),
+					NameText,
+					ModifierStruct->GetDisplayNameText());
+			}
+
+			return FText::Format(LOCTEXT("AudioComponentModifierListDisplayInfoNum_Format", "{0} ({1} modifiers)"),
+				NameText,
+				FText::AsNumber(Modifiers.Num()));
+		}
+
+		return NameText;
+	}
+	
+	FText FSubsonicEventAction_AudioComponentPlay::GetDisplayInfo() const
+	{
+		if (!Name.IsNone())
+		{
+			if (Sound != nullptr)
+			{
+				return FText::Format(LOCTEXT("AudioComponentPlay_DisplayInfoFormat", "{0} ({1})"), FText::FromName(Name), FText::FromName(Sound->GetFName()));
+			}
+
+			return FText::FromName(Name);
+		}
+
+		return Super::GetDisplayInfo();
+	}
+#endif // WITH_EDITOR
+
 	void FSubsonicEventAction_AudioComponentPlay::Execute(const FSubsonicExecutor& InExecutor, const UE::Subsonic::FActionHandle& InHandle) const
 	{
 		const FActionScopeKey ScopeKey(InExecutor, InHandle);
@@ -143,6 +196,18 @@ namespace UE::Subsonic
 		}
 	}
 
+#if WITH_EDITOR
+	FText FSubsonicEventAction_AudioComponentStop::GetDisplayInfo() const
+	{
+		if (!Name.IsNone())
+		{
+			return FText::FromName(Name);
+		}
+
+		return Super::GetDisplayInfo();
+	}
+#endif // WITH_EDITOR
+
 	void FSubsonicEventAction_AudioComponentModifier_Play::Execute(UAudioComponent& AudioComponent) const
 	{
 		AudioComponent.Play(StartTime);
@@ -174,4 +239,5 @@ namespace UE::Subsonic
 	{
 		AudioComponent.Stop();
 	}
-} // namespace UE::Subsonic
\ No newline at end of file
+} // namespace UE::Subsonic
+#undef LOCTEXT_NAMESPACE // SubsonicActions
diff --git a/Engine/Plugins/Experimental/Subsonic/Source/Subsonic/Private/CoreActions/SubsonicAction_AudioComponent.h b/Engine/Plugins/Experimental/Subsonic/Source/Subsonic/Private/CoreActions/SubsonicAction_AudioComponent.h
--- a/Engine/Plugins/Experimental/Subsonic/Source/Subsonic/Private/CoreActions/SubsonicAction_AudioComponent.h
+++ b/Engine/Plugins/Experimental/Subsonic/Source/Subsonic/Private/CoreActions/SubsonicAction_AudioComponent.h
@@ -46,9 +46,13 @@ namespace UE::Subsonic
 
 		virtual void Execute(const FSubsonicExecutor& InExecutor, const UE::Subsonic::FActionHandle& InHandle) const override;
 
+#if WITH_EDITOR
+		virtual FText GetDisplayInfo() const override;
+#endif // WITH_EDITOR
+
 		// Addressable name of audio component instance.
 		UPROPERTY(EditAnywhere, Category = "Subsonic")
-		FName Name = "Sound01";
+		FName Name = "Component01";
 
 		// Scope of component accessor.  If set to event, provided name is used
 		// to address the same component within the event, collection the same component
@@ -70,15 +74,19 @@ namespace UE::Subsonic
 
 		virtual ~FSubsonicEventAction_AudioComponentPlay() = default;
 
+#if WITH_EDITOR
+		virtual FText GetDisplayInfo() const override;
+#endif // WITH_EDITOR
+
 		virtual void Execute(const FSubsonicExecutor& InExecutor, const UE::Subsonic::FActionHandle& InHandle) const override;
 
 	public:
 		// Addressable name of audio component instance.
 		UPROPERTY(EditAnywhere, Category = "Subsonic")
-		FName Name = "Sound01";
+		FName Name = "Component01";
 
 		// Sound to play.
-		UPROPERTY(EditAnywhere, Category = "Subsonic")
+		UPROPERTY(EditAnywhere, Category = "Subsonic", meta = (IsBindable))
 		TObjectPtr<USoundBase> Sound;
 
 		// Scope of component accessor.  If set to event, provided name is used
@@ -100,10 +108,13 @@ namespace UE::Subsonic
 
 		virtual void Execute(const FSubsonicExecutor& InExecutor, const UE::Subsonic::FActionHandle& InHandle) const override;
 
-	public:
+#if WITH_EDITOR
+		virtual FText GetDisplayInfo() const override;
+#endif // WITH_EDITOR
+
 		// Addressable name of component/sound instance.
 		UPROPERTY(EditAnywhere, Category = "Subsonic")
-		FName Name = "Sound01";
+		FName Name = "Component01";
 
 		// Scope of component to play.  If set to event, provided name is used
 		// to address the existing component within the event, collection the existing component
diff --git a/Engine/Plugins/Experimental/Subsonic/Source/Subsonic/Private/CoreActions/SubsonicAction_EventCore.cpp b/Engine/Plugins/Experimental/Subsonic/Source/Subsonic/Private/CoreActions/SubsonicAction_EventCore.cpp
new file mode 100644
--- /dev/null
+++ b/Engine/Plugins/Experimental/Subsonic/Source/Subsonic/Private/CoreActions/SubsonicAction_EventCore.cpp
@@ -0,0 +1,28 @@
+// Copyright Epic Games, Inc. All Rights Reserved.
+
+#include "CoreActions/SubsonicAction_EventCore.h"
+
+#include "GameplayTagContainer.h"
+#include "SubsonicHandles.h"
+
+
+namespace UE::Subsonic
+{
+	void FSubsonicEventAction_DelayEvent::Execute(const FSubsonicExecutor& InExecutor, const UE::Subsonic::FActionHandle& InHandle) const
+	{
+		if (EventName.IsValid())
+		{
+			if (Delay > 0.0f)
+			{
+				InExecutor.ExecuteEvent(EventName.GetTagName());
+			}
+			else
+			{
+				// TODO: Implement delay instance state tracking in subscriber & cycle detection
+				// Detection likely should be implemented in validation step required to execute
+				// in non-runtime contexts for perf reasons.
+				InExecutor.ExecuteEvent(EventName.GetTagName());
+			}
+		}
+	}
+} // namespace UE::Subsonic
\ No newline at end of file
diff --git a/Engine/Plugins/Experimental/Subsonic/Source/Subsonic/Private/CoreActions/SubsonicAction_EventCore.h b/Engine/Plugins/Experimental/Subsonic/Source/Subsonic/Private/CoreActions/SubsonicAction_EventCore.h
new file mode 100644
--- /dev/null
+++ b/Engine/Plugins/Experimental/Subsonic/Source/Subsonic/Private/CoreActions/SubsonicAction_EventCore.h
@@ -0,0 +1,45 @@
+// Copyright Epic Games, Inc. All Rights Reserved.
+#pragma once
+
+#include "GameplayTagContainer.h"
+#include "Misc/CoreMiscDefines.h"
+#include "StructUtils/InstancedStruct.h"
+#include "SubsonicEventCollection.h"
+
+#include "SubsonicAction_EventCore.generated.h"
+
+
+namespace UE::Subsonic
+{
+	UENUM()
+	enum class EEventResolutionRule : uint8
+	{
+		First,
+		Last,
+		Random,
+		Index
+	};
+
+	USTRUCT(MinimalAPI, BlueprintType, DisplayName = "Delay Event")
+	struct FSubsonicEventAction_DelayEvent : public FSubsonicEventActionBase
+	{
+		GENERATED_BODY()
+
+	public:
+		virtual ~FSubsonicEventAction_DelayEvent() = default;
+
+		virtual void Execute(const FSubsonicExecutor& InExecutor, const UE::Subsonic::FActionHandle& InHandle) const override;
+
+		// Identifier used to cancel or replace delay in progress.
+		UPROPERTY(EditAnywhere, BlueprintReadWrite, Category = "Action")
+		FName DelayName = "Delay01";
+
+		// Event to execute after the given delay.
+		UPROPERTY(EditAnywhere, BlueprintReadWrite, Category = "Action")
+		FGameplayTag EventName;
+
+		// Executes the given event by the prescribed amount.
+		UPROPERTY(EditAnywhere, BlueprintReadWrite, Category = "Action")
+		float Delay = 1.0f;
+	};
+} // namespace UE::Subsonic
\ No newline at end of file
diff --git a/Engine/Plugins/Experimental/Subsonic/Source/Subsonic/Private/SubsonicEventCollection.cpp b/Engine/Plugins/Experimental/Subsonic/Source/Subsonic/Private/SubsonicEventCollection.cpp
--- a/Engine/Plugins/Experimental/Subsonic/Source/Subsonic/Private/SubsonicEventCollection.cpp
+++ b/Engine/Plugins/Experimental/Subsonic/Source/Subsonic/Private/SubsonicEventCollection.cpp
@@ -8,29 +8,61 @@
 #include "SubsonicHandles.h"
 #include "SubsonicLog.h"
 #include "SubsonicSubsystem.h"
+#include "UObject/UnrealType.h"
 
 
 namespace UE::Subsonic
 {
-	UE_DEFINE_GAMEPLAY_TAG(TAG_SubsonicEventCorePlay, "Subsonic.Event.Play");
-	UE_DEFINE_GAMEPLAY_TAG(TAG_SubsonicEventCoreStop, "Subsonic.Event.Stop");
+	UE_DEFINE_GAMEPLAY_TAG(TAG_SubsonicCore, "Subsonic");
+	UE_DEFINE_GAMEPLAY_TAG(TAG_SubsonicCore_Event_Play, "Subsonic.Event.Play");
+	UE_DEFINE_GAMEPLAY_TAG(TAG_SubsonicCore_Event_Stop, "Subsonic.Event.Stop");
 
 	namespace CollectionPrivate
 	{
+		constexpr bool bErrorOnMissingTag =
+#if WITH_EDITOR
+		// Validation will catch missing tags. While fixing up tags or
+		// creating new events with an unset tag, this messaging is too pervasive.
+		false;
+#else // !WITH_EDITOR
+		true;
+#endif // !WITH_EDITOR
 		constexpr uint32 InvalidId = static_cast<uint32>(INDEX_NONE);
 	} // namespace CollectionPrivate
 
+	const TInstancedStruct<FSubsonicEventActionBase>& FSubsonicEventActionDefinition::GetAction() const
+	{
+		return Action;
+	}
+
+#if WITH_EDITOR
+	void FSubsonicEventActionDefinition::SetAction(TInstancedStruct<FSubsonicEventActionBase> NewAction)
+	{
+		Action = MoveTemp(NewAction);
+	}
+
+	void FSubsonicEvent::BindParameter(const FInstancedPropertyBag& InSourceParameters, FName ParamName)
+	{
+		for (FSubsonicEventActionDefinition& ActionDefinition : ActionCollection)
+		{
+			if (FSubsonicEventActionBase* ActionPtr = ActionDefinition.Action.GetMutablePtr())
+			{
+				ActionPtr->BindParameter(InSourceParameters, ParamName);
+			}
+		}
+	}
+#endif // WITH_EDITOR
 
 	void FSubsonicEvent::Execute(const FSubsonicExecutor& InExecutor, const FEventHandle& InHandle) const
 	{
 		checkf(InExecutor.GetCollectionHandle().CollectionId == InHandle.Collection.CollectionId, TEXT("Invalid execution context: executor CollectionId must match that of the given event handle"));
 
 		FActionHandle ActionHandle { .Event = InHandle };
-		for (int32 Index = 0; Index < Actions.Num(); ++Index)
+		for (int32 Index = 0; Index < ActionCollection.Num(); ++Index)
 		{
 			ActionHandle.Index = Index;
 
-			TConstStructView<FSubsonicEventActionBase> Action = Actions[Index];
+			TConstStructView<FSubsonicEventActionBase> Action = ActionCollection[Index].Action;
 			if (const FSubsonicEventActionBase* ActionPtr = Action.GetPtr())
 			{
 				ActionPtr->Execute(InExecutor, ActionHandle);
@@ -42,6 +74,42 @@ namespace UE::Subsonic
 		}
 	}
 
+	const TArray<FSubsonicEventActionDefinition>& FSubsonicEvent::GetActionCollection() const
+	{
+		return ActionCollection;
+	}
+
+#if WITH_EDITOR
+	FName FSubsonicEvent::GetActionCollectionPropertyName()
+	{
+		return GET_MEMBER_NAME_CHECKED(FSubsonicEvent, ActionCollection);
+	}
+#endif // WITH_EDITOR
+
+#if WITH_EDITOR
+	bool FSubsonicEvent::GetAutoAudition() const
+	{
+		return bAutoAudition;
+	}
+#endif // WITH_EDITOR
+
+	bool FSubsonicEvent::GetIsPublic() const
+	{
+		return bIsPublic;
+	}
+
+#if WITH_EDITOR
+	TArray<FSubsonicEventActionDefinition>& FSubsonicEvent::GetMutableActionCollection()
+	{
+		return ActionCollection;
+	}
+
+	void FSubsonicEvent::SetIsPublic(bool bInIsPublic)
+	{
+		bIsPublic = bInIsPublic;
+	}
+#endif // WITH_EDITOR
+
 	FSubsonicEventCollectionDefinition::FSubsonicEventCollectionDefinition(FSubsonicEventCollectionDefinition&& InOther)
 	{
 		Unregister();
@@ -82,6 +150,42 @@ namespace UE::Subsonic
 		CollectionId = ++MaxId;
 	}
 
+#if WITH_EDITOR
+	FSubsonicEventActionDefinition* FSubsonicEventCollectionDefinition::AddAction(const FEventHandle& InHandle, int32 InsertIndex)
+	{
+		if (FSubsonicEvent* Event = FindMutableEvent(InHandle))
+		{
+			TArray<FSubsonicEventActionDefinition>& ActionCollection = Event->GetMutableActionCollection();
+			if (!ActionCollection.IsValidIndex(InsertIndex))
+			{
+				return &ActionCollection.AddDefaulted_GetRef();
+			}
+
+			return &ActionCollection.InsertDefaulted_GetRef(InsertIndex);
+		}
+
+		return nullptr;
+	}
+
+	FSubsonicEvent* FSubsonicEventCollectionDefinition::AddEvent(const FGameplayTag& InTag, FSubsonicEvent NewEvent)
+	{
+		if (ContainsEvent(InTag))
+		{
+			return nullptr;
+		}
+
+		return &Events.Add(InTag, MoveTemp(NewEvent));
+	}
+
+	void FSubsonicEventCollectionDefinition::BindParameter(FName ParamName)
+	{
+		for (TPair<FGameplayTag, FSubsonicEvent>& Pair : Events)
+		{
+			Pair.Value.BindParameter(Parameters, ParamName);
+		}
+	}
+#endif // WITH_EDITOR
+
 	FSubsonicEventCollectionDefinition FSubsonicEventCollectionDefinition::Create(FName Name, TMap<FGameplayTag, FSubsonicEvent> InEvents)
 	{
 		FSubsonicEventCollectionDefinition NewDefinition;
@@ -90,9 +194,9 @@ namespace UE::Subsonic
 		const FCollectionHandle NewHandle
 		{
 	#if WITH_EDITORONLY_DATA
-				.CollectionName = Name,
+			.CollectionName = Name,
 	#endif // WITH_EDITORONLY_DATA
-				.CollectionId = NewDefinition.GetCollectionId(),
+			.CollectionId = NewDefinition.GetCollectionId(),
 		};
 
 		NewDefinition.Events = MoveTemp(InEvents);
@@ -106,11 +210,116 @@ namespace UE::Subsonic
 		CollectionId = GetInvalidId();
 	}
 
-	bool FSubsonicEventCollectionDefinition::IsValid() const
+#if WITH_EDITOR
+	bool FSubsonicEventCollectionDefinition::ClearActions(const FEventHandle& InHandle)
 	{
-		return CollectionId != INDEX_NONE;
+		if (FSubsonicEvent* Event = FindMutableEvent(InHandle))
+		{
+			Event->GetMutableActionCollection().Empty();
+			return true;
+		}
+
+		return false;
+	}
+
+	void FSubsonicEventCollectionDefinition::ClearEvents()
+	{
+		Events.Empty();
+	}
+#endif // WITH_EDITOR
+
+	bool FSubsonicEventCollectionDefinition::ContainsEvent(const FGameplayTag& EventTag) const
+	{
+		return Events.Contains(EventTag);
+	}
+
+	bool FSubsonicEventCollectionDefinition::ContainsEvent(FName Name) const
+	{
+		const FGameplayTag EventTag = FGameplayTag::RequestGameplayTag(Name, CollectionPrivate::bErrorOnMissingTag);
+		if (EventTag.GetTagName().IsValid() || Name.IsNone())
+		{
+			return ContainsEvent(EventTag);
+		}
+
+		return false;
+
+	}
+
+	bool FSubsonicEventCollectionDefinition::ContainsEvent(const FEventHandle& InHandle) const
+	{
+		if (InHandle.Collection.CollectionId == CollectionId)
+		{
+			return ContainsEvent(InHandle.EventName);
+		}
+
+		return false;
+	}
+
+	const FSubsonicEventActionDefinition* FSubsonicEventCollectionDefinition::FindAction(const FActionHandle& InHandle) const
+	{
+		if (const FSubsonicEvent* Event = FindEvent(InHandle.Event))
+		{
+			const TArray<FSubsonicEventActionDefinition>& Actions = Event->GetActionCollection();
+			if (Actions.IsValidIndex(InHandle.Index))
+			{
+				return &Actions[InHandle.Index];
+			}
+		}
+
+		return nullptr;
+	}
+
+	const FSubsonicEvent* FSubsonicEventCollectionDefinition::FindEvent(const FGameplayTag& EventTag) const
+	{
+		return Events.Find(EventTag);
+	}
+
+	const FSubsonicEvent* FSubsonicEventCollectionDefinition::FindEvent(FName Name) const
+	{
+		const FGameplayTag EventTag = FGameplayTag::RequestGameplayTag(Name, CollectionPrivate::bErrorOnMissingTag);
+		return FindEvent(EventTag);
+	}
+
+	const FSubsonicEvent* FSubsonicEventCollectionDefinition::FindEvent(const FEventHandle& InHandle) const
+	{
+		if (InHandle.Collection.CollectionId == CollectionId)
+		{
+			return FindEvent(InHandle.EventName);
+		}
+
+		return nullptr;
 	}
 
+#if WITH_EDITOR
+	FSubsonicEventActionDefinition* FSubsonicEventCollectionDefinition::FindMutableAction(const FActionHandle& InHandle)
+	{
+		if (FSubsonicEvent* Event = FindMutableEvent(InHandle.Event))
+		{
+			TArray<FSubsonicEventActionDefinition>& Actions = Event->ActionCollection;
+			if (Actions.IsValidIndex(InHandle.Index))
+			{
+				return &Actions[InHandle.Index];
+			}
+		}
+
+		return nullptr;
+	}
+
+	FSubsonicEvent* FSubsonicEventCollectionDefinition::FindMutableEvent(const FEventHandle& InHandle)
+	{
+		if (InHandle.Collection.CollectionId == GetCollectionId())
+		{
+			const FGameplayTag EventTag = FGameplayTag::RequestGameplayTag(InHandle.EventName, CollectionPrivate::bErrorOnMissingTag);
+			if (EventTag.GetTagName().IsValid() || InHandle.EventName.IsNone())
+			{
+				return Events.Find(EventTag);
+			}
+		}
+
+		return nullptr;
+	}
+#endif // WITH_EDITOR
+
 	uint32 FSubsonicEventCollectionDefinition::GetCollectionId() const
 	{
 		return CollectionId;
@@ -146,6 +355,11 @@ namespace UE::Subsonic
 	}
 #endif // WITH_EDITOR
 
+	bool FSubsonicEventCollectionDefinition::IsValid() const
+	{
+		return CollectionId != INDEX_NONE;
+	}
+
 	void FSubsonicEventCollectionDefinition::Register(const FCollectionHandle& InHandle)
 	{
 		if (GEngine && IsValid())
@@ -157,6 +371,37 @@ namespace UE::Subsonic
 		}
 	}
 
+#if WITH_EDITOR
+	bool FSubsonicEventCollectionDefinition::RemoveAction(const FActionHandle& InHandle)
+	{
+		if (FSubsonicEvent* Event = FindMutableEvent(InHandle.Event))
+		{
+			TArray<FSubsonicEventActionDefinition>& Actions = Event->GetMutableActionCollection();
+			if (Actions.IsValidIndex(InHandle.Index))
+			{
+				Actions.RemoveAt(InHandle.Index);
+				return true;
+			}
+		}
+
+		return false;
+	}
+
+	bool FSubsonicEventCollectionDefinition::RemoveEvent(const FEventHandle& InHandle)
+	{
+		if (InHandle.Collection.CollectionId == CollectionId)
+		{
+			const FGameplayTag EventTag = FGameplayTag::RequestGameplayTag(InHandle.EventName, CollectionPrivate::bErrorOnMissingTag);
+			if (EventTag.GetTagName().IsValid() || InHandle.EventName.IsNone())
+			{
+				return Events.Remove(EventTag) > 0;
+			}
+		}
+
+		return false;
+	}
+#endif // WITH_EDITOR
+
 	void FSubsonicEventCollectionDefinition::Unregister()
 	{
 		if (GEngine && IsValid())
@@ -202,6 +447,14 @@ namespace UE::Subsonic
 		};
 	}
 
+#if WITH_EDITOR
+	FSubsonicEventCollectionDefinition& USubsonicEventCollection::GetMutableDefinition()
+	{
+		return CollectionDefinition;
+	}
+#endif // WITH_EDITOR
+
+
 	void USubsonicEventCollection::SetDefinition(FSubsonicEventCollectionDefinition&& Definition)
 	{
 		CollectionDefinition = MoveTemp(Definition);
@@ -233,6 +486,53 @@ namespace UE::Subsonic
 		Super::BeginDestroy();
 	}
 
+	USubsonicEventCollectionExecutor* USubsonicEventCollectionExecutor::Create(UObject& Parent, FName Name, const USubsonicEventCollection& InCollection, Audio::FDeviceId DeviceId)
+	{
+		if (USubsonicEventCollectionExecutor* NewExecutor = NewObject<USubsonicEventCollectionExecutor>(&Parent, Name, RF_Transient))
+		{
+			class FEventCollectionAccessor : public FSubsonicExecutor::ICollectionAccessor
+			{
+				// Parent USubsonicEventExecutor contains strong ptr, so no need
+				// to induce GC overhead by making this a strong pointer.
+				TWeakObjectPtr<const USubsonicEventCollection> Collection;
+
+			public:
+				FEventCollectionAccessor(const USubsonicEventCollection& InCol)
+					: Collection(TWeakObjectPtr<const USubsonicEventCollection>(&InCol))
+				{
+				}
+
+				virtual ~FEventCollectionAccessor() = default;
+
+				const FSubsonicEventCollectionDefinition* GetDefinition() const override final
+				{
+					if (Collection.IsValid())
+					{
+						return &Collection->GetDefinition();
+					}
+
+					return nullptr;
+				}
+
+				FCollectionHandle GetHandle() const override final
+				{
+					if (Collection.IsValid())
+					{
+						return Collection->GetHandle();
+					}
+
+					return FCollectionHandle();
+				}
+			};
+
+			NewExecutor->Collection = &InCollection;
+			NewExecutor->Executor = FSubsonicExecutor::Create(DeviceId, MakeUnique<FEventCollectionAccessor>(InCollection));
+			return NewExecutor;
+		}
+
+		return nullptr;
+	}
+
 	const USubsonicEventCollection& USubsonicEventCollectionExecutor::GetCollectionChecked() const
 	{
 		check(Collection.Get());
@@ -253,4 +553,9 @@ namespace UE::Subsonic
 	{
 		return Executor.IsValid() && GetCollection() != nullptr;
 	}
+
+	void USubsonicEventCollectionExecutor::Unregister()
+	{
+		Executor.Unregister();
+	}
 } // namespace UE::Subsonic
diff --git a/Engine/Plugins/Experimental/Subsonic/Source/Subsonic/Private/SubsonicExecutor.cpp b/Engine/Plugins/Experimental/Subsonic/Source/Subsonic/Private/SubsonicExecutor.cpp
--- a/Engine/Plugins/Experimental/Subsonic/Source/Subsonic/Private/SubsonicExecutor.cpp
+++ b/Engine/Plugins/Experimental/Subsonic/Source/Subsonic/Private/SubsonicExecutor.cpp
@@ -5,6 +5,7 @@
 #include "Engine/Engine.h"
 #include "SubsonicEventCollection.h"
 #include "SubsonicHandles.h"
+#include "SubsonicLog.h"
 #include "SubsonicSubsystem.h"
 
 
@@ -25,12 +26,12 @@ namespace UE::Subsonic
 	{
 		Unregister();
 
-		CollectionHandle = InOther.CollectionHandle;
-		DeviceId = InOther.DeviceId;
-		ExecutorId = InOther.ExecutorId;
+		CollectionAccessor = MoveTemp(InOther.CollectionAccessor);
 
-		InOther.CollectionHandle = { };
+		DeviceId = InOther.DeviceId;
 		InOther.DeviceId = ExecutorPrivate::InvalidDeviceId;
+
+		ExecutorId = InOther.ExecutorId;
 		InOther.ExecutorId = ExecutorPrivate::InvalidId;
 	}
 
@@ -43,22 +44,22 @@ namespace UE::Subsonic
 	{
 		Unregister();
 
-		CollectionHandle = InOther.CollectionHandle;
-		DeviceId = InOther.DeviceId;
-		ExecutorId = InOther.ExecutorId;
+		CollectionAccessor = MoveTemp(InOther.CollectionAccessor);
 
-		InOther.CollectionHandle = { };
+		DeviceId = InOther.DeviceId;
 		InOther.DeviceId = ExecutorPrivate::InvalidDeviceId;
+
+		ExecutorId = InOther.ExecutorId;
 		InOther.ExecutorId = ExecutorPrivate::InvalidId;
 
 		return *this;
 	}
 
-	FSubsonicExecutor FSubsonicExecutor::Create(Audio::FDeviceId InDeviceId, FCollectionHandle CollectionHandle)
+	FSubsonicExecutor FSubsonicExecutor::Create(Audio::FDeviceId InDeviceId, FSubsonicExecutor::FCollectionAccessorPtr&& CollectionPtr)
 	{
 		static uint32 MaxExecutorId = ExecutorPrivate::InvalidId;
 		FSubsonicExecutor NewHandle;
-		NewHandle.CollectionHandle = MoveTemp(CollectionHandle);
+		NewHandle.CollectionAccessor = MoveTemp(CollectionPtr);
 		NewHandle.ExecutorId = ++MaxExecutorId;
 		NewHandle.DeviceId = InDeviceId;
 		NewHandle.Register();
@@ -67,7 +68,12 @@ namespace UE::Subsonic
 
 	FCollectionHandle FSubsonicExecutor::GetCollectionHandle() const
 	{
-		return CollectionHandle;
+		if (CollectionAccessor.IsValid())
+		{
+			return CollectionAccessor->GetHandle();
+		}
+
+		return FCollectionHandle();
 	}
 
 	Audio::FDeviceId FSubsonicExecutor::GetDeviceId() const
@@ -80,29 +86,43 @@ namespace UE::Subsonic
 		return ExecutorId;
 	}
 
-	void FSubsonicExecutor::ExecuteEvent(const FCollectionHandle& InCollection, FName EventName, const FSubsonicEvent& InEvent) const
+	void FSubsonicExecutor::ExecuteEvent(FName EventName) const
 	{
-		const FEventHandle EventHandle
+		if (CollectionAccessor.IsValid())
 		{
-			.Collection = InCollection,
-			.EventName = EventName
-		};
-
-		if (GEngine)
-		{
-			if (USubsonicSubsystem* Subsystem = GEngine->GetEngineSubsystem<USubsonicSubsystem>())
-			{
-				Subsystem->OnEventPreExecute(*this, EventHandle);
-			}
-		}
-
-		InEvent.Execute(*this, EventHandle);
-
-		if (GEngine)
-		{
-			if (USubsonicSubsystem* Subsystem = GEngine->GetEngineSubsystem<USubsonicSubsystem>())
+			if (const FSubsonicEventCollectionDefinition* Definition = CollectionAccessor->GetDefinition())
 			{
-				Subsystem->OnEventPostExecute(*this, EventHandle);
+				USubsonicSubsystem* Subsystem = GEngine->GetEngineSubsystem<USubsonicSubsystem>();
+				if (!ensureAlwaysMsgf(Subsystem, TEXT("Attempting to execute an event with no initialized SubsonicSubsystem")))
+				{
+					return;
+				}
+
+				const FSubsonicEvent* Event = Definition->FindEvent(EventName);
+				if (!Event)
+				{
+#if !NO_LOGGING
+					const FString CollectionString =
+#if WITH_EDITORONLY_DATA
+						FString::Printf(TEXT(" in collection '%s'"), *CollectionAccessor->GetHandle().CollectionName.ToString());
+#else // !WITH_EDITORONLY_DATA
+						TEXT("");
+#endif // !WITH_EDITORONLY_DATA
+
+					UE_LOG(LogSubsonic, Warning, TEXT("Failed to execute event '%s'%s: Event not found."), *EventName.ToString(), *CollectionString);
+#endif // !NO_LOGGING
+
+					return;
+				}
+
+				if (Event->GetIsPublic())
+				{
+					const FCollectionHandle& CollectionHandle = CollectionAccessor->GetHandle();
+					const FEventHandle EventHandle{ .Collection = CollectionHandle, .EventName = EventName };
+					Subsystem->OnEventPreExecute(*this, EventHandle);
+					Event->Execute(*this, EventHandle);
+					Subsystem->OnEventPostExecute(*this, EventHandle);
+				}
 			}
 		}
 	}
@@ -143,6 +163,7 @@ namespace UE::Subsonic
 			}
 		}
 
+		CollectionAccessor.Reset();
 		ExecutorId = ExecutorPrivate::InvalidId;
 	}
 
diff --git a/Engine/Plugins/Experimental/Subsonic/Source/Subsonic/Private/SubsonicSubsystem.cpp b/Engine/Plugins/Experimental/Subsonic/Source/Subsonic/Private/SubsonicSubsystem.cpp
--- a/Engine/Plugins/Experimental/Subsonic/Source/Subsonic/Private/SubsonicSubsystem.cpp
+++ b/Engine/Plugins/Experimental/Subsonic/Source/Subsonic/Private/SubsonicSubsystem.cpp
@@ -19,17 +19,31 @@ namespace UE::Subsonic
 {
 	namespace SubsystemPrivate
 	{
-		Audio::FDeviceId GetDeviceIdFromWorldObject(const UObject* WorldContextObject)
+		bool TryGetDeviceIdFromWorldObject(const UObject* WorldContextObject, Audio::FDeviceId& OutDeviceId)
 		{
-			if (GEngine)
+			OutDeviceId = INDEX_NONE;
+			if (WorldContextObject)
 			{
-				if (UWorld* World = GEngine->GetWorldFromContextObject(WorldContextObject, EGetWorldErrorMode::LogAndReturnNull))
+				if (GEngine)
 				{
-					return World->GetAudioDevice().GetDeviceID();
+					if (UWorld* World = GEngine->GetWorldFromContextObject(WorldContextObject, EGetWorldErrorMode::LogAndReturnNull))
+					{
+						OutDeviceId = World->GetAudioDevice().GetDeviceID();
+						return true;
+					}
+				}
+			}
+			else
+			{
+				if (FAudioDeviceManager* DeviceManager = FAudioDeviceManager::Get())
+				{
+					OutDeviceId = DeviceManager->GetActiveAudioDevice().GetDeviceID();
+					return true;
 				}
 			}
 
-			return INDEX_NONE;
+			// Either engine isn't loaded, or context object has no parent world
+			return false;
 		};
 
 		TOptional<Audio::FDeviceId> GetDeviceIdFromSubscriber(const TScriptInterface<ISubsonicEventSubscriberInterface>& Subscriber)
@@ -45,8 +59,13 @@ namespace UE::Subsonic
 				return Subsystem->GetAudioDeviceId();
 			}
 
-			return GetDeviceIdFromWorldObject(Subscriber.GetObject());
+			Audio::FDeviceId DeviceId = INDEX_NONE;
+			if (TryGetDeviceIdFromWorldObject(Subscriber.GetObject(), DeviceId))
+			{
+				return DeviceId;
+			}
 
+			return { };
 		}
 
 		class FSubsonicSubsystemImpl : public FGCObject
@@ -166,20 +185,15 @@ namespace UE::Subsonic
 		Super::Deinitialize();
 	}
 
-	USubsonicEventCollectionExecutor* USubsonicSubsystem::CreateExecutorBP(UObject* WorldContextObject, const USubsonicEventCollection* Collection)
+	USubsonicEventCollectionExecutor* USubsonicSubsystem::CreateExecutorBP(UObject* WorldContextObject, FName Name, const USubsonicEventCollection* InCollection)
 	{
-		USubsonicEventCollectionExecutor* NewExecutor = NewObject<USubsonicEventCollectionExecutor>(WorldContextObject);
-
-		if (GEngine && NewExecutor && Collection)
+		if (InCollection && WorldContextObject)
 		{
-			const Audio::FDeviceId DeviceId = SubsystemPrivate::GetDeviceIdFromWorldObject(WorldContextObject);
-
-			NewExecutor->Collection = Collection;
-
-			FCollectionHandle CollectionHandle = Collection->GetHandle();
-			NewExecutor->Executor = FSubsonicExecutor::Create(DeviceId, MoveTemp(CollectionHandle));
-
-			return NewExecutor;
+			Audio::FDeviceId DeviceId = INDEX_NONE;
+			if (SubsystemPrivate::TryGetDeviceIdFromWorldObject(WorldContextObject, DeviceId))
+			{
+				return USubsonicEventCollectionExecutor::Create(*WorldContextObject, Name, *InCollection, DeviceId);
+			}
 		}
 
 		return nullptr;
@@ -193,8 +207,7 @@ namespace UE::Subsonic
 			const FSubsonicEventCollectionDefinition& Definition = EventCollection.GetDefinition();
 			if (const FSubsonicEvent* Event = Definition.GetEvents().Find(EventTag))
 			{
-				const FCollectionHandle CollectionHandle = EventCollection.GetHandle();
-				Executor->GetExecutor().ExecuteEvent(CollectionHandle, EventTag.GetTagName(), *Event);
+				Executor->GetExecutor().ExecuteEvent(EventTag.GetTagName());
 			}
 		}
 	}
diff --git a/Engine/Plugins/Experimental/Subsonic/Source/Subsonic/Public/SubsonicEventCollection.h b/Engine/Plugins/Experimental/Subsonic/Source/Subsonic/Public/SubsonicEventCollection.h
--- a/Engine/Plugins/Experimental/Subsonic/Source/Subsonic/Public/SubsonicEventCollection.h
+++ b/Engine/Plugins/Experimental/Subsonic/Source/Subsonic/Public/SubsonicEventCollection.h
@@ -13,6 +13,9 @@
 
 #include "SubsonicEventCollection.generated.h"
 
+// Forward Declarations
+class FProperty;
+
 #define UE_API SUBSONIC_API
 
 namespace UE::Subsonic
@@ -20,8 +23,18 @@ namespace UE::Subsonic
 	// Forward Declarations
 	struct FActionHandle;
 
-	UE_API UE_DECLARE_GAMEPLAY_TAG_EXTERN(TAG_SubsonicEventCorePlay)
-	UE_API UE_DECLARE_GAMEPLAY_TAG_EXTERN(TAG_SubsonicEventCoreStop)
+	UE_API UE_DECLARE_GAMEPLAY_TAG_EXTERN(TAG_SubsonicCore)
+	UE_API UE_DECLARE_GAMEPLAY_TAG_EXTERN(TAG_SubsonicCore_Event_Play)
+	UE_API UE_DECLARE_GAMEPLAY_TAG_EXTERN(TAG_SubsonicCore_Event_Stop)
+
+	USTRUCT()
+	struct FSubsonicEventActionBoundProperties
+	{
+		GENERATED_BODY()
+
+		UPROPERTY()
+		TArray<FName> Properties;
+	};
 
 	USTRUCT(MinimalAPI, BlueprintType, meta = (Abstract, Hidden))
 	struct FSubsonicEventActionBase
@@ -30,10 +43,44 @@ namespace UE::Subsonic
 
 		virtual ~FSubsonicEventActionBase() = default;
 
+#if WITH_EDITOR
+		// Returns display info that is custom to the given action instance to be shown
+		// adjacent to the action, for example, in the Subsonic Event Editor action tree.
+		UE_API virtual FText GetDisplayInfo() const { return FText(); }
+#endif // WITH_EDITOR
+
 	protected:
+#if WITH_EDITOR
+		virtual void BindParameter(const FInstancedPropertyBag& InBag, FName Name) { }
+#endif // WITH_EDITOR
+
 		virtual void Execute(const FSubsonicExecutor& InExecutor, const FActionHandle& InHandle) const { }
 
+#if WITH_EDITORONLY_DATA
+		UPROPERTY()
+		TMap<FName, FSubsonicEventActionBoundProperties> Bindings;
+#endif // WITH_EDITORONLY_DATA
+
+		friend struct FSubsonicEvent;
+	};
+
+	// Action definition iteratively executed per Subsonic event.
+	USTRUCT(MinimalAPI, BlueprintType)
+	struct FSubsonicEventActionDefinition
+	{
+		GENERATED_BODY()
+
+		UE_API const TInstancedStruct<FSubsonicEventActionBase>& GetAction() const;
+
+#if WITH_EDITOR
+		UE_API void SetAction(TInstancedStruct<FSubsonicEventActionBase> NewAction);
+#endif // WITH_EDITOR
+ 
+	private:
 		friend struct FSubsonicEvent;
+
+		UPROPERTY(EditAnywhere, BlueprintReadOnly, Category = "Actions", meta = (AllowPrivateAccess))
+		TInstancedStruct<FSubsonicEventActionBase> Action;
 	};
 
 	USTRUCT(MinimalAPI, BlueprintType)
@@ -43,16 +90,50 @@ namespace UE::Subsonic
 
 		virtual ~FSubsonicEvent() = default;
 
+#if WITH_EDITOR
+		UE_API UE_INTERNAL static FName GetActionCollectionPropertyName();
+#endif // WITH_EDITOR
+
+		UE_API const TArray<FSubsonicEventActionDefinition>& GetActionCollection() const;
+
+#if WITH_EDITOR
+		UE_API bool GetAutoAudition() const;
+#endif // WITH_EDITOR
+
+		UE_API bool GetIsPublic() const;
+
+#if WITH_EDITOR
+		UE_API UE_INTERNAL TArray<FSubsonicEventActionDefinition>& GetMutableActionCollection();
+		UE_API UE_INTERNAL void SetIsPublic(bool bInIsPublic);
+#endif // WITH_EDITOR
+
 	private:
+		friend struct FSubsonicEventCollectionDefinition;
 		friend struct FSubsonicExecutor;
 
+#if WITH_EDITOR
+		UE_INTERNAL void BindParameter(const FInstancedPropertyBag& InSourceParameters, FName Name);
+#endif // WITH_EDITOR
+
+	private:
 		// Called by executor, iterates through all actions on the given event definition.
 		// Asserts if the provided executor and handle have mismatched collection contexts.
 		UE_INTERNAL void Execute(const FSubsonicExecutor& InExecutor, const FEventHandle& InHandle) const;
 
-	public:
-		UPROPERTY(EditAnywhere, Category = "Subsonic")
-		TArray<TInstancedStruct<FSubsonicEventActionBase>> Actions;
+		UPROPERTY(EditAnywhere, BlueprintReadOnly, Category = "Event", meta = (AllowPrivateAccess))
+		bool bIsPublic = true;
+
+#if WITH_EDITORONLY_DATA
+		// If true, execute this event as soon as an audition executor is created
+		UPROPERTY(EditAnywhere, Category = "Event", meta = (EditCondition = bIsPublic))
+		bool bAutoAudition = false;
+
+		UPROPERTY(EditAnywhere, Category = "Event")
+		FInstancedPropertyBag Parameters;
+#endif // WITH_EDITORONLY_DATA
+
+		UPROPERTY(EditAnywhere, BlueprintReadOnly, Category = "Actions", meta = (AllowPrivateAccess))
+		TArray<FSubsonicEventActionDefinition> ActionCollection;
 	};
 
 
@@ -79,8 +160,42 @@ namespace UE::Subsonic
 		// outside of UObject reflection.
 		static FSubsonicEventCollectionDefinition Create(FName Name, TMap<FGameplayTag, FSubsonicEvent> InEvents);
 
-		UE_API bool IsValid() const;
+#if WITH_EDITOR
+		// Adds action to event with handle if it exists. If insert index set to less than 0 (default), adds to end of collection array.
+		UE_API UE_INTERNAL FSubsonicEventActionDefinition* AddAction(const FEventHandle& InHandle, int32 InsertIndex = INDEX_NONE);
+
+		// Adds event with given tag if does not yet exists. Returns pointer to new event if added, null if none added.
+		UE_API UE_INTERNAL FSubsonicEvent* AddEvent(const FGameplayTag& InTag = { }, FSubsonicEvent NewEvent = {});
+
+		// Clears all actions under the given event
+		UE_API UE_INTERNAL bool ClearActions(const FEventHandle& InHandle);
+
+		// Removes all events in the given collection.
+		UE_API UE_INTERNAL void ClearEvents();
+#endif // WITH_EDITOR
+
+		UE_API bool ContainsEvent(FName Name) const;
+		UE_API bool ContainsEvent(const FGameplayTag& EventTag) const;
+		UE_API bool ContainsEvent(const FEventHandle& InHandle) const;
+
+		// Returns pointer to action at handle if it exists
+		UE_API const FSubsonicEventActionDefinition* FindAction(const FActionHandle& InHandle) const;
+
+		// Returns event with the given name.
+		UE_API const FSubsonicEvent* FindEvent(FName Name) const;
+		UE_API const FSubsonicEvent* FindEvent(const FGameplayTag& EventTag) const;
+		UE_API const FSubsonicEvent* FindEvent(const FEventHandle& InHandle) const;
+
+#if WITH_EDITOR
+		// Returns pointer to action at handle if it exists
+		UE_API UE_INTERNAL FSubsonicEventActionDefinition* FindMutableAction(const FActionHandle& InHandle);
+
+		// Returns mutable pointer to event at handle if it exists
+		UE_API UE_INTERNAL FSubsonicEvent* FindMutableEvent(const FEventHandle& InHandle);
+#endif // WITH_EDITOR
+
 		UE_API uint32 GetCollectionId() const;
+
 		UE_API const TMap<FGameplayTag, FSubsonicEvent>& GetEvents() const;
 
 		// Returns Id of invalid (default constructed or unregistered/reset collection).
@@ -91,7 +206,27 @@ namespace UE::Subsonic
 		UE_API static FName GetParametersPropertyName();
 #endif // WITH_EDITOR
 
+		UE_API bool IsValid() const;
+
+#if WITH_EDITOR
+		// Removes action at given handle and returns if successful.
+		UE_API UE_INTERNAL bool RemoveAction(const FActionHandle& InHandle);
+
+		// Removes event at given handle and returns if successful.
+		UE_API UE_INTERNAL bool RemoveEvent(const FEventHandle& InHandle);
+#endif // WITH_EDITOR
+
 	private:
+#if WITH_EDITOR
+		template <typename ValueType>
+		bool GetParameterValue(FName ParamName, ValueType& OutValue) const
+		{
+			return Parameters.GetValue(ParamName, OutValue);
+		}
+
+		UE_INTERNAL void BindParameter(FName Parameter);
+#endif // WITH_EDITOR
+
 		// Register's the given collection with the Subsonic Subsystem.
 		UE_INTERNAL void Register(const FCollectionHandle& InHandle);
 
@@ -114,9 +249,6 @@ namespace UE::Subsonic
 		FInstancedPropertyBag Parameters;
 #endif // WITH_EDITORONLY_DATA
 
-// 		UPROPERTY(VisibleAnywhere, Category = "Subsonic", AdvancedDisplay)
-// 		FSubsonicParameterBindingCollection Bindings;
-
 		void AssignId();
 
 		friend class USubsonicEventCollection;
@@ -135,6 +267,11 @@ namespace UE::Subsonic
 #endif // WITH_EDITOR
 
 		UE_API FCollectionHandle GetHandle() const;
+
+#if WITH_EDITOR
+		UE_API UE_INTERNAL FSubsonicEventCollectionDefinition& GetMutableDefinition();
+#endif // WITH_EDITOR
+
 		UE_API void SetDefinition(FSubsonicEventCollectionDefinition&& Definition);
 
 	private:
@@ -153,13 +290,28 @@ namespace UE::Subsonic
 		GENERATED_BODY()
 
 	public:
-		const USubsonicEventCollection* GetCollection() const;
-		const USubsonicEventCollection& GetCollectionChecked() const;
-		const FSubsonicExecutor& GetExecutor() const;
-		bool IsValid() const;
-
 		virtual void BeginDestroy() override;
 
+		// Creates a transient executor that is bound to the given collection.
+		UE_API static USubsonicEventCollectionExecutor* Create(UObject& Parent, FName Name, const USubsonicEventCollection& InCollection, Audio::FDeviceId DeviceId);
+
+		// Returns collection executor is bound to.
+		UE_API const USubsonicEventCollection* GetCollection() const;
+
+		// Returns collection executor is bound to and asserts if no collection is set.
+		UE_API const USubsonicEventCollection& GetCollectionChecked() const;
+
+		// Returns underlying executor UObject wraps.
+		UE_API const FSubsonicExecutor& GetExecutor() const;
+
+		// Returns whether this executor is valid (has a valid ID and is bound to an event collection)
+		UE_API bool IsValid() const;
+
+		// Unregisters the underlying executor implementation, clearing
+		// any state pertaining to the given executor managed by subscribers
+		// and releasing association with the provided collection when created.
+		UE_API void Unregister();
+
 	private:
 		UPROPERTY(Transient, BlueprintReadOnly, Category = "Subsonic", meta = (AllowPrivateAccess))
 		TObjectPtr<const USubsonicEventCollection> Collection;
diff --git a/Engine/Plugins/Experimental/Subsonic/Source/Subsonic/Public/SubsonicExecutor.h b/Engine/Plugins/Experimental/Subsonic/Source/Subsonic/Public/SubsonicExecutor.h
--- a/Engine/Plugins/Experimental/Subsonic/Source/Subsonic/Public/SubsonicExecutor.h
+++ b/Engine/Plugins/Experimental/Subsonic/Source/Subsonic/Public/SubsonicExecutor.h
@@ -7,6 +7,7 @@
 #include "Misc/CoreDefines.h"
 #include "SubsonicHandles.h"
 #include "Templates/TypeHash.h"
+#include "Templates/UniquePtr.h"
 #include "UObject/NameTypes.h"
 
 #include "SubsonicExecutor.generated.h"
@@ -31,6 +32,9 @@ enum class ESubsonicExecutionScope : uint8
 
 namespace UE::Subsonic
 {
+	// Forward Declarations
+	class USubsonicSubsystem;
+
 	// Execution object that tracks explicit execution context for any caller of events related to a given collection.
 	// RAII design that automatically registers and unregisters with the Subsonic Subsystem, and by extension,
 	// any implemented and active subscribers.  This enables subscribers to track and manage data related to
@@ -44,19 +48,32 @@ namespace UE::Subsonic
 	public:
 		FSubsonicExecutor(FSubsonicExecutor&&);
 		FSubsonicExecutor& operator=(const FSubsonicExecutor& InOther) = delete;
-		~FSubsonicExecutor();
+		UE_API ~FSubsonicExecutor();
 
 		UE_API FSubsonicExecutor& operator=(FSubsonicExecutor&& InOther);
 
-		UE_API static FSubsonicExecutor Create(Audio::FDeviceId InDeviceId, FCollectionHandle CollectionHandle);
+		// Interface for providing a collection definition. Abstracted to enable
+		// supporting either UObject reflection ownership or as a generic native implementation
+		// of Executor.
+		struct ICollectionAccessor
+		{
+			virtual ~ICollectionAccessor() = default;
+
+			virtual const FSubsonicEventCollectionDefinition* GetDefinition() const = 0;
+			virtual FCollectionHandle GetHandle() const = 0;
+		};
+		using FCollectionAccessorPtr = TUniquePtr<ICollectionAccessor>;
+
+		UE_API static FSubsonicExecutor Create(Audio::FDeviceId InDeviceId, FCollectionAccessorPtr&& CollectionPtr);
+
+		UE_API void ExecuteEvent(FName EventName) const;
+
 		UE_API static uint32 GetInvalidId();
 		UE_API static FSubsonicExecutor GetInvalid();
 
-		FCollectionHandle GetCollectionHandle() const;
-		Audio::FDeviceId GetDeviceId() const;
-		uint32 GetId() const;
-
-		void ExecuteEvent(const FCollectionHandle& InCollection, FName EventName, const FSubsonicEvent& InEvent) const;
+		UE_API FCollectionHandle GetCollectionHandle() const;
+		UE_API Audio::FDeviceId GetDeviceId() const;
+		UE_API uint32 GetId() const;
 
 		UE_API bool IsValid() const;
 
@@ -71,7 +88,8 @@ namespace UE::Subsonic
 		// Registers Executor with the SubsonicSubsystem (and by extension all subscribers).
 		void Register();
 
-		FCollectionHandle CollectionHandle;
+		// Pointer to associated collection accessor.
+		FCollectionAccessorPtr CollectionAccessor;
 
 		// DeviceId corresponding to which audio device to execute event on.
 		Audio::FDeviceId DeviceId;
diff --git a/Engine/Plugins/Experimental/Subsonic/Source/Subsonic/Public/SubsonicSubsystem.h b/Engine/Plugins/Experimental/Subsonic/Source/Subsonic/Public/SubsonicSubsystem.h
--- a/Engine/Plugins/Experimental/Subsonic/Source/Subsonic/Public/SubsonicSubsystem.h
+++ b/Engine/Plugins/Experimental/Subsonic/Source/Subsonic/Public/SubsonicSubsystem.h
@@ -40,7 +40,7 @@ namespace UE::Subsonic
 		//~ End USubsystem interface
 
 		UFUNCTION(BlueprintCallable, Category = "Subsonic", meta = (DisplayName = "Create Executor", WorldContext = "WorldContextObject"))
-		UE_API UPARAM(DisplayName = "Executor") USubsonicEventCollectionExecutor* CreateExecutorBP(UObject* WorldContextObject, const USubsonicEventCollection* Collection);
+		UE_API UPARAM(DisplayName = "Executor") USubsonicEventCollectionExecutor* CreateExecutorBP(UObject* WorldContextObject, FName Name, const USubsonicEventCollection* Collection);
 
 		UFUNCTION(BlueprintCallable, Category = "Subsonic", meta = (DisplayName = "Execute Event"))
 		UE_API void ExecuteEventBP(const USubsonicEventCollectionExecutor* Executor, FGameplayTag Event);
@@ -50,17 +50,17 @@ namespace UE::Subsonic
 		friend struct FSubsonicEventCollectionDefinition;
 		friend class USubsonicEventSubscriberBase;
 
-		UE_INTERNAL void OnCollectionRegistered(const FCollectionHandle& InHandle);
-		UE_INTERNAL void OnCollectionUnregistered(const FCollectionHandle& InHandle);
+		void OnCollectionRegistered(const FCollectionHandle& InHandle);
+		void OnCollectionUnregistered(const FCollectionHandle& InHandle);
 
-		UE_INTERNAL void OnEventPreExecute(const FSubsonicExecutor& Executor, const FEventHandle& InEventHandle);
-		UE_INTERNAL void OnEventPostExecute(const FSubsonicExecutor& Executor, const FEventHandle& InEventHandle);
+		void OnEventPreExecute(const FSubsonicExecutor& Executor, const FEventHandle& InEventHandle);
+		void OnEventPostExecute(const FSubsonicExecutor& Executor, const FEventHandle& InEventHandle);
 
-		UE_INTERNAL void OnEventSubscriberRegistered(TScriptInterface<ISubsonicEventSubscriberInterface> Subscriber);
-		UE_INTERNAL void OnEventSubscriberUnregistered(TScriptInterface<ISubsonicEventSubscriberInterface> Subscriber);
+		void OnEventSubscriberRegistered(TScriptInterface<ISubsonicEventSubscriberInterface> Subscriber);
+		void OnEventSubscriberUnregistered(TScriptInterface<ISubsonicEventSubscriberInterface> Subscriber);
 
-		UE_INTERNAL void OnExecutorRegistered(const FSubsonicExecutor& InExecutor);
-		UE_INTERNAL void OnExecutorUnregistered(const FSubsonicExecutor& InExecutor);
+		void OnExecutorRegistered(const FSubsonicExecutor& InExecutor);
+		void OnExecutorUnregistered(const FSubsonicExecutor& InExecutor);
 
 		TPimplPtr<SubsystemPrivate::FSubsonicSubsystemImpl> Impl;
 	};
diff --git a/Engine/Plugins/Experimental/Subsonic/Source/SubsonicEditor/Private/SubsonicAssetDefinitions.cpp b/Engine/Plugins/Experimental/Subsonic/Source/SubsonicEditor/Private/SubsonicAssetDefinitions.cpp
--- a/Engine/Plugins/Experimental/Subsonic/Source/SubsonicEditor/Private/SubsonicAssetDefinitions.cpp
+++ b/Engine/Plugins/Experimental/Subsonic/Source/SubsonicEditor/Private/SubsonicAssetDefinitions.cpp
@@ -13,6 +13,37 @@
 
 namespace UE::Subsonic
 {
+	namespace AssetDefinitionsPrivate
+	{
+		const FSlateBrush& GetSlateBrushSafe(FName InName)
+		{
+			const ISlateStyle* Style = FSlateStyleRegistry::FindSlateStyle("SubsonicStyle");
+			if (ensureMsgf(Style, TEXT("Missing slate style 'SubsonicStyle'")))
+			{
+				const FSlateBrush* Brush = Style->GetBrush(InName);
+				if (ensureMsgf(Brush, TEXT("Missing brush '%s'"), *InName.ToString()))
+				{
+					return *Brush;
+				}
+			}
+
+			if (const FSlateBrush* NoBrush = FAppStyle::GetBrush("NoBrush"))
+			{
+				return *NoBrush;
+			}
+
+			static const FSlateBrush NullBrush;
+			return NullBrush;
+		}
+
+		const FSlateBrush* GetClassBrush(const FAssetData& InAssetData, FName InClassName, bool bIsThumbnail = false)
+		{
+			FString BrushName = FString::Printf(TEXT("SubsonicEditor.%s"), *InClassName.ToString());
+			BrushName += bIsThumbnail ? TEXT(".Thumbnail") : TEXT(".Icon");
+			return &GetSlateBrushSafe(*BrushName);
+		}
+	} // namespace AssetDefinitionsPrivate
+
 	FLinearColor UAssetDefinition_SubsonicEventCollection::GetAssetColor() const
 	{
  		if (const ISlateStyle* SubsonicStyle = FSlateStyleRegistry::FindSlateStyle("SubsonicStyle"))
@@ -52,11 +83,21 @@ namespace UE::Subsonic
 
 	const FSlateBrush* UAssetDefinition_SubsonicEventCollection::GetThumbnailBrush(const FAssetData& InAssetData, const FName InClassName) const
 	{
+		if (const FSlateBrush* ClassBrush = AssetDefinitionsPrivate::GetClassBrush(InAssetData, InClassName, true /* bIsThumbnail */))
+		{
+			return ClassBrush;
+		}
+
 		return Super::GetThumbnailBrush(InAssetData, InClassName);
 	}
 
 	const FSlateBrush* UAssetDefinition_SubsonicEventCollection::GetIconBrush(const FAssetData& InAssetData, const FName InClassName) const
 	{
+		if (const FSlateBrush* ClassBrush = AssetDefinitionsPrivate::GetClassBrush(InAssetData, InClassName))
+		{
+			return ClassBrush;
+		}
+
 		return Super::GetIconBrush(InAssetData, InClassName);
 	}
 
diff --git a/Engine/Plugins/Experimental/Subsonic/Source/SubsonicEditor/Private/SubsonicEditorModule.cpp b/Engine/Plugins/Experimental/Subsonic/Source/SubsonicEditor/Private/SubsonicEditorModule.cpp
--- a/Engine/Plugins/Experimental/Subsonic/Source/SubsonicEditor/Private/SubsonicEditorModule.cpp
+++ b/Engine/Plugins/Experimental/Subsonic/Source/SubsonicEditor/Private/SubsonicEditorModule.cpp
@@ -51,8 +51,8 @@ namespace UE::Subsonic::Editor
 			// Register Custom Collection Class Views
 			{
 				FPropertyEditorModule& PropertyModule = FModuleManager::LoadModuleChecked<FPropertyEditorModule>("PropertyEditor");
-				CollectionViewEntries.Add(FCollectionViewEntry<USubsonicEventsView, FCollectionEventsDetailCustomization>::Register(PropertyModule));
-				CollectionViewEntries.Add(FCollectionViewEntry<USubsonicParametersView, FCollectionParametersDetailCustomization>::Register(PropertyModule));
+				CollectionViewEntries.Add(FCollectionViewEntry<USubsonicEventTreeDetailsView, FEventTreeSelectionDetailCustomization>::Register(PropertyModule));
+				CollectionViewEntries.Add(FCollectionViewEntry<USubsonicCollectionParametersView, FCollectionParametersDetailCustomization>::Register(PropertyModule));
 			}
 
 			StyleSet = MakeShared<FSlateStyle>();
diff --git a/Engine/Plugins/Experimental/Subsonic/Source/SubsonicEditor/Private/SubsonicEditorSubsystem.cpp b/Engine/Plugins/Experimental/Subsonic/Source/SubsonicEditor/Private/SubsonicEditorSubsystem.cpp
new file mode 100644
--- /dev/null
+++ b/Engine/Plugins/Experimental/Subsonic/Source/SubsonicEditor/Private/SubsonicEditorSubsystem.cpp
@@ -0,0 +1,47 @@
+// Copyright Epic Games, Inc. All Rights Reserved.
+
+#include "SubsonicEditorSubsystem.h"
+
+#include "SubsonicEventCollection.h"
+#include "UObject/UObjectIterator.h"
+
+
+namespace UE::Subsonic
+{
+	void USubsonicEditorSubsystem::Initialize(FSubsystemCollectionBase& Collection)
+	{
+		Super::Initialize(Collection);
+		RebuildActionStructChildCache();
+	}
+
+	void USubsonicEditorSubsystem::Deinitialize()
+	{
+		Super::Deinitialize();
+	}
+
+	void USubsonicEditorSubsystem::RebuildActionStructChildCache()
+	{
+		ActionStructs.Empty();
+		for (TObjectIterator<UScriptStruct> It; It; ++It)
+		{
+			if (UScriptStruct* Struct = *It)
+			{
+				if (Struct->IsChildOf(FSubsonicEventActionBase::StaticStruct()))
+				{
+					ActionStructs.Add(TWeakObjectPtr<UScriptStruct>(Struct));
+				}
+			}
+		}
+	}
+
+	void USubsonicEditorSubsystem::ForEachActionStruct(TFunctionRef<void(const UScriptStruct& Struct)> StructIter) const
+	{
+		for (const TWeakObjectPtr<const UScriptStruct>& Class : ActionStructs)
+		{
+			if (const UScriptStruct* Struct = Class.Get())
+			{
+				StructIter(*Struct);
+			}
+		}
+	}
+} // namespace UE::Subsonic
diff --git a/Engine/Plugins/Experimental/Subsonic/Source/SubsonicEditor/Private/SubsonicEditorSubsystem.h b/Engine/Plugins/Experimental/Subsonic/Source/SubsonicEditor/Private/SubsonicEditorSubsystem.h
new file mode 100644
--- /dev/null
+++ b/Engine/Plugins/Experimental/Subsonic/Source/SubsonicEditor/Private/SubsonicEditorSubsystem.h
@@ -0,0 +1,39 @@
+// Copyright Epic Games, Inc. All Rights Reserved.
+#pragma once
+
+#include "Containers/Array.h"
+#include "EditorSubsystem.h"
+#include "Templates/Function.h"
+#include "UObject/Class.h"
+#include "UObject/WeakObjectPtrTemplates.h"
+
+#include "SubsonicEditorSubsystem.generated.h"
+
+
+namespace UE::Subsonic
+{
+	/**
+	 * Interface to the Subsonic Editor Functionality.
+	 */
+	UCLASS(MinimalAPI, BlueprintType)
+	class USubsonicEditorSubsystem final : public UEditorSubsystem
+	{
+		GENERATED_BODY()
+
+	public:
+		void Initialize(FSubsystemCollectionBase& Collection) override;
+		void Deinitialize() override;
+
+		// Rebuilds the action struct child cache (called on initialize,
+		// but available to refresh if new modules are loaded or unloaded).
+		// Cached to allow for faster iteration over projects with large numbers
+		// of registered struct types.
+		void RebuildActionStructChildCache();
+
+		// Iterates over cached action struct child pointers.
+		void ForEachActionStruct(TFunctionRef<void(const UScriptStruct& Struct)> StructIter) const;
+
+	private:
+		TArray<TWeakObjectPtr<const UScriptStruct>> ActionStructs;
+	};
+} // namespace UE::Subsonic
diff --git a/Engine/Plugins/Experimental/Subsonic/Source/SubsonicEditor/Private/SubsonicEventCollectionEditor.cpp b/Engine/Plugins/Experimental/Subsonic/Source/SubsonicEditor/Private/SubsonicEventCollectionEditor.cpp
--- a/Engine/Plugins/Experimental/Subsonic/Source/SubsonicEditor/Private/SubsonicEventCollectionEditor.cpp
+++ b/Engine/Plugins/Experimental/Subsonic/Source/SubsonicEditor/Private/SubsonicEventCollectionEditor.cpp
@@ -2,14 +2,28 @@
 
 #include "SubsonicEventCollectionEditor.h"
 
+#include "AudioDeviceHandle.h"
+#include "AudioDeviceManager.h"
+#include "Engine/Engine.h"
 #include "Framework/Application/SlateApplication.h"
 #include "Framework/Docking/TabManager.h"
-#include "PropertyEditorModule.h"
+#include "GameplayTagContainer.h"
+#include "ISinglePropertyView.h"
 #include "Modules/ModuleManager.h"
-#include "Widgets/Docking/SDockTab.h"
+#include "PropertyBindingExtension.h"
+#include "PropertyEditorModule.h"
+#include "ScopedTransaction.h"
+#include "SGameplayTagCombo.h"
 #include "Styling/SlateStyleRegistry.h"
+#include "SubsonicEditorSubsystem.h"
 #include "SubsonicEventCollection.h"
-#include "SubsonicEventCollectionViews.h"
+#include "SubsonicEventCollectionEditorCommands.h"
+#include "Widgets/Docking/SDockTab.h"
+#include "Widgets/Input/SSearchableComboBox.h"
+#include "Widgets/Input/SSearchBox.h"
+#include "Widgets/SNullWidget.h"
+#include "Widgets/Views/SListView.h"
+#include "Widgets/Views/STableRow.h"
 
 #define LOCTEXT_NAMESPACE "SubsonicEditor"
 
@@ -18,39 +32,569 @@ namespace UE::Subsonic::Editor
 {
 	namespace EventCollectionEditorPrivate
 	{
-		const FName TabName_Events = "SubsonicEventCollectionEditor_Events";
+		const FName TabName_EventTree = "SubsonicEventCollectionEditor_EventTree";
+		const FName TabName_Inspector = "SubsonicEventCollectionEditor_Inspector";
 		const FName TabName_Parameters = "SubsonicEventCollectionEditor_Parameters";
+
+		Audio::FDeviceId GetPreviewDeviceId()
+		{
+			if (FAudioDeviceManager* DeviceManager = GEngine->GetAudioDeviceManager())
+			{
+				if (FAudioDeviceHandle ActiveDevice = DeviceManager->GetActiveAudioDevice(); ActiveDevice.IsValid())
+				{
+					return ActiveDevice->DeviceID;
+				}
+			}
+
+			return INDEX_NONE;
+		}
+
+		const FSlateBrush& GetSlateBrushSafe(FName InName)
+		{
+			const ISlateStyle* MetaSoundStyle = FSlateStyleRegistry::FindSlateStyle("SubsonicStyle");
+			if (ensureMsgf(MetaSoundStyle, TEXT("Missing slate style 'SubsonicStyle'")))
+			{
+				const FSlateBrush* Brush = MetaSoundStyle->GetBrush(InName);
+				if (ensureMsgf(Brush, TEXT("Missing brush '%s'"), *InName.ToString()))
+				{
+					return *Brush;
+				}
+			}
+
+			if (const FSlateBrush* NoBrush = FAppStyle::GetBrush("NoBrush"))
+			{
+				return *NoBrush;
+			}
+
+			static const FSlateBrush NullBrush;
+			return NullBrush;
+		}
+
+		TSharedPtr<IPropertyHandle> GetEventTagPropertyHandle(TSharedRef<ISinglePropertyView> CollectionView, FName EventName)
+		{
+			TSharedPtr<IPropertyHandle> CollectionProperty = CollectionView->GetPropertyHandle();
+			if (!CollectionProperty)
+			{
+				return { };
+			}
+
+			TSharedPtr<IPropertyHandle> EventsHandle = CollectionProperty->GetChildHandle(FSubsonicEventCollectionDefinition::GetEventsPropertyName());
+			if (!EventsHandle)
+			{
+				return { };
+			}
+
+			TSharedPtr<IPropertyHandleMap> EventMapHandle = EventsHandle->AsMap();
+			if (!EventMapHandle)
+			{
+				return { };
+			}
+
+			uint32 NumEvents = 0;
+			if (EventMapHandle->GetNumElements(NumEvents) != FPropertyAccess::Success || NumEvents <= 0)
+			{
+				return { };
+			}
+
+			for (uint32 EventIndex = 0; EventIndex < NumEvents; ++EventIndex)
+			{
+				TSharedRef<IPropertyHandle> EventEntryHandle = EventMapHandle->GetElement(EventIndex);
+				if (TSharedPtr<IPropertyHandle> EventKeyHandle = EventEntryHandle->GetKeyHandle(); EventKeyHandle.IsValid())
+				{
+					FName Name;
+					TSharedPtr<IPropertyHandle> TagNameHandle = EventKeyHandle->GetChildHandle("TagName");
+					if (TagNameHandle.IsValid()
+						&& TagNameHandle->GetValue(Name) == FPropertyAccess::Success
+						&& EventName == Name)
+					{
+						return EventKeyHandle;
+					}
+				}
+			}
+
+			return { };
+		}
 	} // namespace EventCollectionEditorPrivate
 
 	const FName FEventCollectionEditor::EditorName = "SubsonicEventCollectionEditor";
 
+	TSharedRef<ITableRow> FEventCollectionEditor::OnEventTree_GenerateRow(TSharedPtr<FEventCollectionEditor::FHandleVariant> HandleVariant, const TSharedRef<STableViewBase>& TableView)
+	{
+		using STableRowPtr = TSharedPtr<STableRow<TSharedPtr<FHandleVariant>>>;
+		STableRowPtr NewRow = SNew(STableRow<TSharedPtr<FHandleVariant>>, TableView);
+
+		const USubsonicEventCollection* Collection = GetEventCollection();
+		if (!HandleVariant.IsValid() || !Collection)
+		{
+			return NewRow.ToSharedRef();
+		}
+
+		constexpr float IndentMult = 4.f;
+		FText RowDisplayName = FText::FromName("None"); // Valid display name for actions with no struct type set
+		float RowIndent = 5.f;
+		if (HandleVariant->IsAction())
+		{
+			if (const FSubsonicEventActionDefinition* Action = Collection->GetDefinition().FindAction(HandleVariant->GetAction()))
+			{
+				RowIndent *= IndentMult;
+				if (const UScriptStruct* ActionStruct = Action->GetAction().GetScriptStruct())
+				{
+					RowDisplayName = ActionStruct->GetDisplayNameText();
+				}
+			}
+		}
+		else if (HandleVariant->IsEvent())
+		{
+			RowDisplayName = FText::FromName(HandleVariant->GetEvent().EventName);
+		}
+		else
+		{
+			checkNoEntry();
+		}
+
+		TSharedPtr<SWidget> NameWidget = SNullWidget::NullWidget;
+		if (HandleVariant->IsEvent())
+		{
+			if (CollectionView.IsValid())
+			{
+				TSharedPtr<IPropertyHandle> EventNameProperty = EventCollectionEditorPrivate::GetEventTagPropertyHandle(CollectionView.ToSharedRef(), HandleVariant->GetEvent().EventName);
+				if (EventNameProperty.IsValid())
+				{
+					NameWidget = SNew(SGameplayTagCombo)
+					.Filter(TAG_SubsonicCore.GetTag().GetTagName().ToString())
+					.PropertyHandle(EventNameProperty)
+					.OnTagChanged_Lambda([this, HandleRef = HandleVariant.ToSharedRef()](const FGameplayTag NewTag)
+					{
+						if (NewTag.GetTagName() != HandleRef->GetEvent().EventName)
+						{
+							FEventHandle NewEvent = HandleRef->GetEvent();
+							NewEvent.EventName = NewTag.GetTagName();
+							HandleRef->Set(MoveTemp(NewEvent));
+							RebuildEventTree(true /* bRebuildListView */);
+						}
+					});
+				}
+			}
+		}
+		else if (const FSubsonicEventActionDefinition* Action = Collection->GetDefinition().FindAction(HandleVariant->GetAction()))
+		{
+			FText ActionDisplayName = FText::FromName("None");
+			if (const UScriptStruct* ActionStruct = Action->GetAction().GetScriptStruct())
+			{
+				ActionDisplayName = ActionStruct->GetDisplayNameText();
+			}
+
+			TSharedPtr<SSearchableComboBox> ComboBox = SNew(SSearchableComboBox)
+			.OptionsSource(&ActionComboValues)
+			.OnGenerateWidget_Lambda([this](TSharedPtr<FString> InItem) -> TSharedRef<SWidget>
+			{
+				TSharedRef<SWidget> EntryWidget = SNullWidget::NullWidget;
+				if (InItem.IsValid())
+				{
+					const FString& ItemStr = *InItem.Get();
+					if (USubsonicEditorSubsystem* EditorSubsystem = GEditor->GetEditorSubsystem<USubsonicEditorSubsystem>())
+					{
+						EditorSubsystem->ForEachActionStruct([this, &EntryWidget, &ItemStr](const UScriptStruct& ActionStruct)
+						{
+							if (ActionStruct.GetFName() == ItemStr)
+							{
+								EntryWidget = SNew(STextBlock).Text(ActionStruct.GetDisplayNameText());
+							}
+						});
+					}
+				}
+
+				return EntryWidget;
+			})
+			.OnSelectionChanged_Lambda([this, HandleRef = HandleVariant.ToSharedRef()](TSharedPtr<FString> NewSelection, ESelectInfo::Type SelectInfo)
+			{
+				if (SelectInfo != ESelectInfo::OnNavigation && NewSelection.IsValid())
+				{
+					TransactEventCollection(FText::Format(LOCTEXT("SetActionDescription", "Set action type to '{0}'"), FText::FromString(*NewSelection)), [&](FSubsonicEventCollectionDefinition& Definition)
+					{
+						if (FSubsonicEventActionDefinition* Action = Definition.FindMutableAction(HandleRef->GetAction()))
+						{
+							const FName NewSelectionName = FName(*NewSelection);
+							if (USubsonicEditorSubsystem* EditorSubsystem = GEditor->GetEditorSubsystem<USubsonicEditorSubsystem>())
+							{
+								EditorSubsystem->ForEachActionStruct([this, &Action, NewSelectionName](const UScriptStruct& ActionStruct)
+								{
+									if (ActionStruct.GetFName() == NewSelectionName)
+									{
+										TInstancedStruct<FSubsonicEventActionBase> NewConfig;
+										NewConfig.InitializeAsScriptStruct(&ActionStruct);
+										Action->SetAction(MoveTemp(NewConfig));
+										RebuildEventTree(true /* bRebuildListView */);
+									}
+								});
+							}
+						}
+					});
+				}
+			})
+			.Content()
+			[
+				SNew(STextBlock).Text(ActionDisplayName)
+			];
+
+			SAssignNew(NameWidget, SHorizontalBox)
+			+ SHorizontalBox::Slot()
+			.VAlign(VAlign_Center)
+				.AutoWidth()
+			.Padding(0.f, 0.f, 5.0f, 0.f)
+			[
+				ComboBox.ToSharedRef()
+			]
+			+ SHorizontalBox::Slot()
+			.VAlign(VAlign_Center)
+			.AutoWidth()
+			.Padding(0.f, 0.f, 5.0f, 0.f)
+			[
+				SNew(STextBlock)
+				.Font(FCoreStyle::GetDefaultFontStyle("Italic", 10))
+				.Text_Lambda([this, HandleRef = HandleVariant.ToSharedRef()]()
+				{
+					if (const USubsonicEventCollection* Collection = GetEventCollection())
+					{
+						if (const FSubsonicEventActionDefinition* Action = Collection->GetDefinition().FindAction(HandleRef->GetAction()))
+						{
+							if (const FSubsonicEventActionBase* ActionBase = Action->GetAction().GetPtr())
+							{
+								return ActionBase->GetDisplayInfo();
+							}
+						}
+					}
+					return FText();
+				})
+			];
+		}
+
+		TSharedPtr<SHorizontalBox> RowBox;
+		TSharedRef<SHorizontalBox> NameBox = SNew(SHorizontalBox)
+		+ SHorizontalBox::Slot()
+			.VAlign(VAlign_Center)
+			.FillWidth(1.0f)
+			.Padding(0.f, 0.f, 5.0f, 0.f)
+		[
+			NameWidget.ToSharedRef()
+		];
+
+		auto VisibleOnRowHover = [RowPtr = NewRow.ToWeakPtr()]()
+		{
+			if (RowPtr.IsValid() && RowPtr.Pin()->IsHovered())
+			{
+				return EVisibility::Visible;
+			}
+			return EVisibility::Hidden;
+		};
+		TSharedRef<FHandleVariant> HandleRef = HandleVariant.ToSharedRef();
+
+		if (HandleVariant->IsEvent())
+		{
+			NameBox->AddSlot()
+			.VAlign(VAlign_Center)
+			.AutoWidth()
+			[
+				SNew(SButton)
+				.ButtonStyle(FAppStyle::Get(), "SimpleButton")
+				.OnClicked_Lambda([this, HandleRef]()
+				{
+					TransactEvent(HandleRef->GetEvent(), LOCTEXT("ToggleEventAccessor", "Toggle event {0} accessor"), [](FSubsonicEvent& Event)
+					{
+						const bool bWasPublic = Event.GetIsPublic();
+						Event.SetIsPublic(!bWasPublic);
+					});
+					return FReply::Handled();
+				})
+				.ToolTipText_Lambda([this, HandleRef]()
+				{
+					if (const USubsonicEventCollection* Collection = GetEventCollection())
+					{
+						if (const FSubsonicEvent* Event = Collection->GetDefinition().FindEvent(HandleRef->GetEvent()))
+						{
+							if (Event->GetIsPublic())
+							{
+								return LOCTEXT("IsEventPublic_Tooltip", "Event can be executed from outside of internal actions.");
+							}
+						}
+					}
+					return LOCTEXT("IsEventPrivate_Tooltip", "Event can only be executed from actions contained in this event collection (not from external scripts or other collections).");
+				})
+				.ForegroundColor(FSlateColor::UseForeground())
+				.ContentPadding(0)
+				.IsFocusable(false)
+				[
+					SNew(SImage)
+					.Image_Lambda([this, HandleRef]()
+					{
+						if (const USubsonicEventCollection* Collection = GetEventCollection())
+						{
+							if (const FSubsonicEvent* Event = Collection->GetDefinition().FindEvent(HandleRef->GetEvent()); Event && Event->GetIsPublic())
+							{
+								return FAppStyle::GetBrush("Kismet.VariableList.ExposeForInstance");
+							}
+						}
+
+						return FAppStyle::GetBrush("Kismet.VariableList.HideForInstance");
+					})
+					.ColorAndOpacity(FSlateColor::UseForeground())
+					.DesiredSizeOverride(FVector2D(16.f, 16.f))
+				]
+			];
+			NameBox->AddSlot()
+			.VAlign(VAlign_Center)
+			.AutoWidth()
+			[
+				SNew(SButton)
+				.ButtonStyle(FAppStyle::Get(), "SimpleButton")
+				.OnClicked_Lambda([this, HandleRef]()
+				{
+					ExecutePreviewEvent(HandleRef->GetEvent());
+					return FReply::Handled();
+				})
+				.ToolTipText(LOCTEXT("Test Execute", "Execute the given event using the editor world as a test."))
+				.ForegroundColor(FSlateColor::UseForeground())
+				.ContentPadding(0)
+				.IsFocusable(false)
+				[
+					SNew(SImage)
+					.Image(&EventCollectionEditorPrivate::GetSlateBrushSafe("SubsonicEditor.Play.Small"))
+					.ColorAndOpacity(FSlateColor::UseForeground())
+					.DesiredSizeOverride(FVector2D(16.f, 16.f))
+					.Visibility_Lambda(VisibleOnRowHover)
+				]
+			];
+		}
+
+		NewRow->SetContent(
+			SNew(SBox)
+			.HAlign(HAlign_Fill)
+			.VAlign(VAlign_Fill)
+			.Padding(RowIndent, 5.0f, 5.0f, 5.0f)
+			[
+				SAssignNew(RowBox, SHorizontalBox)
+				+ SHorizontalBox::Slot()
+				.HAlign(HAlign_Left)
+				.VAlign(VAlign_Fill)
+				[ NameBox ]
+			]
+		);
+
+		RowBox->AddSlot()
+		.HAlign(HAlign_Right)
+		.VAlign(VAlign_Center)
+		[
+			SNew(SHorizontalBox)
+			+ SHorizontalBox::Slot()
+			[
+				SNew(SButton)
+				.ButtonStyle(FAppStyle::Get(), "SimpleButton")
+				.OnClicked_Lambda([this, HandleRef]()
+				{
+					InsertAction(*HandleRef);
+					RebuildEventTree(true /* bRebuildListView */);
+					return FReply::Handled();
+				})
+				.ToolTipText_Lambda([this, HandleRef]()
+				{
+					return FText::Format(
+						HandleRef->IsEvent()
+							? LOCTEXT("AddActionToEndofEventToolTip_Format", "Add action to end of event '{0}' array")
+							: LOCTEXT("InsertActionInEventToolTip_Format", "Insert new action in event '{0}' array"),
+						FText::FromName(HandleRef->GetEvent().EventName));
+				})
+				.ForegroundColor(FSlateColor::UseForeground())
+				.ContentPadding(0)
+				.IsFocusable(false)
+				.Visibility_Lambda(VisibleOnRowHover)
+				[
+					SNew(SImage)
+					.Image(FAppStyle::Get().GetBrush("Icons.AddCircle"))
+					.ColorAndOpacity(FSlateColor::UseForeground())
+				]
+			]
+			+ SHorizontalBox::Slot()
+			[
+				SNew(SButton)
+				.ButtonStyle(FAppStyle::Get(), "SimpleButton")
+				.OnClicked_Lambda([this, HandleRef]()
+				{
+					const FString HandleDesc = HandleRef->IsAction() ? HandleRef->GetAction().ToString() : HandleRef->GetEvent().ToString();
+					FText TransactionDesc = FText::Format(LOCTEXT("RemoveHandleDescription", "Remove {0} from Subsonic Event Collection"), FText::FromString(HandleDesc));
+					TransactEventCollection(TransactionDesc, [&](FSubsonicEventCollectionDefinition& Definition)
+					{
+						const FEventHandle& EventHandle = HandleRef->GetEvent();
+						bool bAddedHandleItem = HandleRef->IsEvent()
+							? Definition.RemoveEvent(EventHandle)
+							: Definition.RemoveAction(HandleRef->GetAction());
+						if (bAddedHandleItem)
+						{
+							RebuildEventTree(true /* bRebuildListView */);
+						}
+					});
+					return FReply::Handled();
+				})
+				.ToolTipText(HandleRef->IsEvent()
+					? FText::Format(LOCTEXT("DeleteEventToolTip_Format", "Delete event '{0}' from Subsonic Event Collection"), FText::FromName(HandleRef->GetEvent().EventName))
+					: LOCTEXT("DeleteActionToolTip_Format", "Delete action from event"))
+				.ForegroundColor(FSlateColor::UseForeground())
+				.ContentPadding(0)
+				.Visibility_Lambda(VisibleOnRowHover)
+				.IsFocusable(false)
+				[
+					SNew(SImage)
+					.Image(FAppStyle::Get().GetBrush("Icons.Delete"))
+					.ColorAndOpacity(FSlateColor::UseForeground())
+				]
+			]
+		];
+
+		return NewRow.ToSharedRef();
+	}
+
+	void FEventCollectionEditor::OnEventTree_RowSelected(TSharedPtr<FHandleVariant> HandleVariant, ESelectInfo::Type SelectInfo)
+	{
+		if (TreeDetailsView.IsValid() && HandleVariant.IsValid())
+		{
+			if (HandleVariant->IsAction())
+			{
+				const FActionHandle ActionHandle = HandleVariant->GetAction();
+				TreeDetailsView->Event = ActionHandle.Event.EventName;
+				TreeDetailsView->ActionIndex = ActionHandle.Index;
+			}
+			else
+			{
+				TreeDetailsView->Event = HandleVariant->GetEvent().EventName;
+				TreeDetailsView->ActionIndex = INDEX_NONE;
+			}
+		}
+
+		if (EventTreeInspectorView.IsValid())
+		{
+			EventTreeInspectorView->ForceRefresh();
+		}
+	}
+
+	void FEventCollectionEditor::RebuildEventTree(bool bRebuildListView)
+	{
+		ActionComboValues.Reset();
+		EventTreeHandles.Reset();
+
+		if (!GEditor)
+		{
+			return;
+		}
+
+		USubsonicEventCollection* EventCollection = GetEventCollection();
+		if (!EventCollection)
+		{
+			return;
+		}
+
+		FPropertyEditorModule& PropertyEditor = FModuleManager::LoadModuleChecked<FPropertyEditorModule>("PropertyEditor");
+		CollectionView = PropertyEditor.CreateSingleProperty(EventCollection, USubsonicEventCollection::GetDefinitionPropertyName(), { });
+
+		if (USubsonicEditorSubsystem* EditorSubsystem = GEditor->GetEditorSubsystem<USubsonicEditorSubsystem>())
+		{
+			EditorSubsystem->ForEachActionStruct([this](const UScriptStruct& ActionStruct)
+			{
+				if (!ActionStruct.HasMetaData("Hidden"))
+				{
+					ActionComboValues.Add(MakeShared<FString>(ActionStruct.GetFName().ToString()));
+				}
+			});
+		}
+
+		using FTagEventPair = TPair<FGameplayTag, FSubsonicEvent>;
+		using FTagEventPtrPair = TPair<FGameplayTag, const FSubsonicEvent*>;
+		TArray<FTagEventPtrPair> SortedEvents;
+		Algo::TransformIf(EventCollection->GetDefinition().GetEvents(), SortedEvents,
+		[this](const FTagEventPair& Pair)
+		{
+			if (EventViewFilter.IsEmpty())
+			{
+				return true;
+			}
+			if (Pair.Key.GetTagName().ToString().Contains(EventViewFilter))
+			{
+				return true;
+			}
+			return Pair.Value.GetActionCollection().ContainsByPredicate([this](const FSubsonicEventActionDefinition& Action)
+			{
+				const UScriptStruct* Struct = Action.GetAction().GetScriptStruct();
+				return Struct && Struct->GetDisplayNameText().ToString().Contains(EventViewFilter);
+			});
+		},
+		[this](const FTagEventPair& Pair)
+		{
+			return FTagEventPtrPair(Pair.Key, &Pair.Value);
+		});
+
+		SortedEvents.Sort([](const FTagEventPtrPair& A, const FTagEventPtrPair& B)
+		{
+			return A.Key.GetTagName().LexicalLess(B.Key.GetTagName());
+		});
+		for (const FTagEventPtrPair& Pair : SortedEvents)
+		{
+			FHandleVariant EventVariant;
+			const FEventHandle EventHandle = { .Collection = EventCollection->GetHandle(), .EventName = Pair.Key.GetTagName() };
+			EventVariant.Set(EventHandle);
+			EventTreeHandles.Add(MakeShared<FHandleVariant>(MoveTemp(EventVariant)));
+
+			int32 ActionIndex = 0;
+			for (const FSubsonicEventActionDefinition& ActionDefinition : Pair.Value->GetActionCollection())
+			{
+				FHandleVariant ActionVariant;
+				const FActionHandle ActionHandle = { .Event = EventHandle, .Index = ActionIndex++ };
+				ActionVariant.Set(ActionHandle);
+				EventTreeHandles.Add(MakeShared<FHandleVariant>(MoveTemp(ActionVariant)));
+			}
+		}
+
+		if (bRebuildListView && EventTreeView.IsValid())
+		{
+			EventTreeView->RebuildList();
+		}
+	}
+
 	void FEventCollectionEditor::RegisterTabSpawners(const TSharedRef<FTabManager>& InTabManager)
 	{
 		WorkspaceMenuCategory = InTabManager->AddLocalWorkspaceMenuCategory(LOCTEXT("WorkspaceMenu_SubsonicEventCollectionEditor", "Subsonic EventCollection Editor"));
 
-		FDetailsViewArgs Args;
-		Args.bHideSelectionTip = true;
-		Args.NotifyHook = this;
+		FDetailsViewArgs DetailViewArgs;
+		DetailViewArgs.bHideSelectionTip = true;
+		DetailViewArgs.NotifyHook = this;
 
 		FPropertyEditorModule& PropertyModule = FModuleManager::LoadModuleChecked<FPropertyEditorModule>("PropertyEditor");
 
-		EventsDetailsView = PropertyModule.CreateDetailView(Args);
-		InTabManager->RegisterTabSpawner(EventCollectionEditorPrivate::TabName_Events, FOnSpawnTab::CreateLambda([ViewRef = EventsDetailsView.ToSharedRef()](const FSpawnTabArgs& Args)
-			{
-				return SNew(SDockTab).Label(LOCTEXT("SubsonicEventCollectionEventsTitle", "Events"))[ ViewRef ];
-			}))
-			.SetDisplayName(LOCTEXT("EventsTab", "Events"))
-			.SetGroup(WorkspaceMenuCategory.ToSharedRef())
-			.SetIcon(FSlateIcon(FAppStyle::GetAppStyleSetName(), "LevelEditor.Tabs.Details"));
+		ParametersDetailsView = PropertyModule.CreateDetailView(DetailViewArgs);
+		InTabManager->RegisterTabSpawner(EventCollectionEditorPrivate::TabName_Parameters, FOnSpawnTab::CreateLambda([ViewRef = ParametersDetailsView.ToSharedRef()](const FSpawnTabArgs&)
+		{
+			return SNew(SDockTab).Label(LOCTEXT("SubsonicEventCollectionParametersTitle", "Parameters")) [ ViewRef ];
+		}))
+		.SetDisplayName(LOCTEXT("ParametersTab", "Parameters"))
+		.SetGroup(WorkspaceMenuCategory.ToSharedRef())
+		.SetIcon(FSlateIcon(FAppStyle::GetAppStyleSetName(), "LevelEditor.Tabs.Details"));
+
+		InTabManager->RegisterTabSpawner(EventCollectionEditorPrivate::TabName_EventTree, FOnSpawnTab::CreateLambda([ViewRef = InitEventTreeView()](const FSpawnTabArgs& Args)
+		{
+			return SNew(SDockTab).Label(LOCTEXT("SubsonicEventCollectionEventsTreeTitle", "Event Tree")) [ ViewRef ];
+		}))
+		.SetDisplayName(LOCTEXT("EventTreeTab", "Event Tree"))
+		.SetGroup(WorkspaceMenuCategory.ToSharedRef())
+		.SetIcon(FSlateIcon(FAppStyle::GetAppStyleSetName(), "LevelEditor.Tabs.Outliner"));
 
-		ParametersDetailsView = PropertyModule.CreateDetailView(Args);
-		InTabManager->RegisterTabSpawner(EventCollectionEditorPrivate::TabName_Parameters, FOnSpawnTab::CreateLambda([ViewRef = ParametersDetailsView.ToSharedRef()](const FSpawnTabArgs& Args)
-			{
-				return SNew(SDockTab).Label(LOCTEXT("SubsonicEventCollectionParametersTitle", "Parameters"))[ ViewRef ];
-			}))
-			.SetDisplayName(LOCTEXT("ParametersTab", "Parameters"))
-			.SetGroup(WorkspaceMenuCategory.ToSharedRef())
-			.SetIcon(FSlateIcon(FAppStyle::GetAppStyleSetName(), "LevelEditor.Tabs.Details"));
+
+		DetailViewArgs.bAllowSearch = false;
+		EventTreeInspectorView = PropertyModule.CreateDetailView(DetailViewArgs);
+		InTabManager->RegisterTabSpawner(EventCollectionEditorPrivate::TabName_Inspector, FOnSpawnTab::CreateLambda([ViewRef = EventTreeInspectorView.ToSharedRef()](const FSpawnTabArgs&)
+		{
+			return SNew(SDockTab).Label(LOCTEXT("SubsonicEventCollectionInspectorTitle", "Inspector")) [ ViewRef ];
+		}))
+		.SetDisplayName(LOCTEXT("InspectorTab", "Inspector"))
+		.SetGroup(WorkspaceMenuCategory.ToSharedRef())
+		.SetIcon(FSlateIcon(FAppStyle::GetAppStyleSetName(), "LevelEditor.Tabs.Properties"));
+		EventTreeInspectorView->SetExtensionHandler(MakeShared<FPropertyBindingExtension>());
 
 		FAssetEditorToolkit::RegisterTabSpawners(InTabManager);
 	}
@@ -59,12 +603,15 @@ namespace UE::Subsonic::Editor
 	{
 		FAssetEditorToolkit::UnregisterTabSpawners(InTabManager);
 
-		InTabManager->UnregisterTabSpawner(EventCollectionEditorPrivate::TabName_Events);
+		InTabManager->UnregisterTabSpawner(EventCollectionEditorPrivate::TabName_EventTree);
+		InTabManager->UnregisterTabSpawner(EventCollectionEditorPrivate::TabName_Inspector);
 		InTabManager->UnregisterTabSpawner(EventCollectionEditorPrivate::TabName_Parameters);
 	}
 
 	FEventCollectionEditor::~FEventCollectionEditor()
 	{
+		StopAudition();
+
 		if (GEditor)
 		{
 			GEditor->UnregisterForUndo(this);
@@ -86,12 +633,17 @@ namespace UE::Subsonic::Editor
 			->SetHideTabWell(true)
 			->AddTab(EventCollectionEditorPrivate::TabName_Parameters, ETabState::OpenedTab);
 
-		EventsStack = FTabManager::NewStack()
-			->SetSizeCoefficient(0.67f)
+		EventTreeStack = FTabManager::NewStack()
+			->SetSizeCoefficient(0.33f)
+			->SetHideTabWell(true)
+			->AddTab(EventCollectionEditorPrivate::TabName_EventTree, ETabState::OpenedTab);
+
+		DetailsStack = FTabManager::NewStack()
+			->SetSizeCoefficient(0.34f)
 			->SetHideTabWell(true)
-			->AddTab(EventCollectionEditorPrivate::TabName_Events, ETabState::OpenedTab);
+			->AddTab(EventCollectionEditorPrivate::TabName_Inspector, ETabState::OpenedTab);
 
-		const TSharedRef<FTabManager::FLayout> StandaloneDefaultLayout = FTabManager::NewLayout("Standalone_SubsonicEventCollectionEditor_Layout_v2")
+		const TSharedRef<FTabManager::FLayout> StandaloneDefaultLayout = FTabManager::NewLayout("Standalone_SubsonicEventCollectionEditor_Layout_v7")
 			->AddArea
 			(
 				FTabManager::NewPrimaryArea()
@@ -102,10 +654,17 @@ namespace UE::Subsonic::Editor
 				)
 				->Split
 				(
-					EventsStack.ToSharedRef()
+					EventTreeStack.ToSharedRef()
+				)
+				->Split
+				(
+					DetailsStack.ToSharedRef()
 				)
 			);
 
+		FEventCollectionEditorCommands::Register();
+		BuildTransport();
+
 		{
 			constexpr bool bCreateDefaultStandaloneMenu = true;
 			constexpr bool bCreateDefaultToolbar = true;
@@ -124,22 +683,174 @@ namespace UE::Subsonic::Editor
 				bUseSmallToolbarIcons);
 		}
 
-		USubsonicEventCollection* EventCollection = CastChecked<USubsonicEventCollection>(GetEditingObject());
-		if (EventsDetailsView.IsValid())
-		{
-			EventsView.Reset(NewObject<USubsonicEventsView>());
-			EventsView->SetCollection(EventCollection);
-			EventsDetailsView->SetObject(EventsView.Get());
-		}
+		USubsonicEventCollection* EventCollection = GetEventCollection();
+		check(EventCollection);
+
+		RebuildEventTree();
 
 		if (ParametersDetailsView.IsValid())
 		{
-			ParametersView.Reset(NewObject<USubsonicParametersView>());
+			ParametersView.Reset(NewObject<USubsonicCollectionParametersView>());
 			ParametersView->SetCollection(EventCollection);
 			ParametersDetailsView->SetObject(ParametersView.Get());
 		}
 
+		if (EventTreeInspectorView.IsValid())
+		{
+			TreeDetailsView.Reset(NewObject<USubsonicEventTreeDetailsView>());
+			TreeDetailsView->SetCollection(EventCollection);
+			EventTreeInspectorView->SetObject(TreeDetailsView.Get());
+		}
+
 		RegenerateMenusAndToolbars();
+		const FEventCollectionEditorCommands& Commands = FEventCollectionEditorCommands::Get();
+		ToolkitCommands->MapAction(
+			Commands.StartAudition,
+			FExecuteAction::CreateSP(this, &FEventCollectionEditor::StartAudition));
+
+		ToolkitCommands->MapAction(
+			Commands.StopAudition,
+			FExecuteAction::CreateSP(this, &FEventCollectionEditor::StopAudition));
+
+		ToolkitCommands->MapAction(
+			Commands.ToggleAudition,
+			FExecuteAction::CreateSP(this, &FEventCollectionEditor::ToggleAudition));
+	}
+
+	TSharedRef<SWidget> FEventCollectionEditor::InitEventTreeView()
+	{
+		TSharedRef<SHeaderRow> EventTreeHeader = SNew(SHeaderRow);
+		EventTreeHeader->AddColumn(
+			SHeaderRow::FColumn::FArguments()
+			.ColumnId("Events")
+			.DefaultTooltip(LOCTEXT("EventsColumn_Tooltip", "Events that can be executed at runtime by code or script."))
+			.VAlignHeader(VAlign_Center)
+			.VAlignCell(VAlign_Center)
+			.HAlignHeader(HAlign_Left)
+			.HAlignCell(HAlign_Fill)
+			.HeaderContentPadding(FMargin(10.0f, 5.0f))
+			.HeaderContent()
+			[
+				SNew(SHorizontalBox)
+				+ SHorizontalBox::Slot()
+				.VAlign(VAlign_Center)
+				.HAlign(HAlign_Fill)
+				.FillWidth(1.0f)
+				[
+					SNew(STextBlock)
+					.TextStyle(FAppStyle::Get(), "DetailsView.CategoryTextStyle")
+					.Text(LOCTEXT("EventsColumn_DisplayName", "Events"))
+				]
+				+ SHorizontalBox::Slot()
+				.VAlign(VAlign_Center)
+				.AutoWidth()
+				.HAlign(HAlign_Right)
+				[
+					SNew(SButton)
+					.ButtonStyle(FAppStyle::Get(), "SimpleButton")
+					.OnClicked_Lambda([this]()
+					{
+						TransactEventCollection(LOCTEXT("AddEventDescription", "Add new Subsonic event to collection"), [&](FSubsonicEventCollectionDefinition& Definition)
+						{
+							if (FSubsonicEvent* NewEvent = Definition.AddEvent())
+							{
+								RebuildEventTree(true /* bRebuildListView */);
+							}
+						});
+						return FReply::Handled();
+					})
+					.ToolTipText(LOCTEXT("AddEventToolTip", "Add event to collection"))
+					.ForegroundColor(FSlateColor::UseForeground())
+					.ContentPadding(0)
+					.IsFocusable(false)
+					[
+						SNew(SImage)
+						.Image(FAppStyle::Get().GetBrush("Icons.AddCircle"))
+						.ColorAndOpacity(FSlateColor::UseForeground())
+					]
+				]
+				+ SHorizontalBox::Slot()
+				.VAlign(VAlign_Center)
+				.HAlign(HAlign_Right)
+				.AutoWidth()
+				[
+					SNew(SButton)
+					.ButtonStyle(FAppStyle::Get(), "SimpleButton")
+					.OnClicked_Lambda([this]()
+					{
+						TransactEventCollection(LOCTEXT("ClearAllEventsDescription", "Clear all Subsonic Collection Events"), [&](FSubsonicEventCollectionDefinition& Definition)
+						{
+							Definition.ClearEvents();
+							RebuildEventTree(true /* bRebuildListView */);
+						});
+						return FReply::Handled();
+					})
+					.ToolTipText(LOCTEXT("DeleteAllEventsToolTip", "Delete all events in collection"))
+					.ForegroundColor(FSlateColor::UseForeground())
+					.ContentPadding(0)
+					.IsFocusable(false)
+					[
+						SNew(SImage)
+						.Image(FAppStyle::Get().GetBrush("Icons.Delete"))
+						.ColorAndOpacity(FSlateColor::UseForeground())
+					]
+				]
+			]
+		);
+
+		auto SetFilterString = [this](const FText& InFilterText)
+		{
+			EventViewFilter = InFilterText.ToString();
+			RebuildEventTree(true /* bRebuildListView */);
+		};
+		return SNew(SVerticalBox)
+		+ SVerticalBox::Slot()
+		.AutoHeight()
+		[
+			SNew(SHorizontalBox)
+			+ SHorizontalBox::Slot()
+			.Padding(6.f)
+			.FillWidth(1.0f)
+			[
+				SNew(SSearchBox)
+				.HintText(LOCTEXT("SearchDetailsHint", "Search"))
+				.OnTextChanged_Lambda(SetFilterString)
+// 				.OnTextCommitted_Lambda(SetFilterString)
+				.AddMetaData<FTagMetaData>(TEXT("Details.Search"))
+			]
+		]
+		+ SVerticalBox::Slot()
+		.AutoHeight()
+		[
+			SAssignNew(EventTreeView, SListView<TSharedPtr<FHandleVariant>>)
+			.HeaderRow(EventTreeHeader)
+			.ListItemsSource(&EventTreeHandles)
+			.OnGenerateRow(this, &FEventCollectionEditor::OnEventTree_GenerateRow)
+			.OnSelectionChanged(this, &FEventCollectionEditor::OnEventTree_RowSelected)
+		];
+	}
+
+	void FEventCollectionEditor::TransactEvent(const FEventHandle EventHandle, FText Description, TFunctionRef<void(FSubsonicEvent&)> TransactionFunc) const
+	{
+		if (USubsonicEventCollection* Collection = GetEventCollection())
+		{
+			if (FSubsonicEvent* Event = Collection->GetMutableDefinition().FindMutableEvent(EventHandle))
+			{
+				FScopedTransaction Transaction(Description);
+				Collection->Modify();
+				TransactionFunc(*Event);
+			}
+		}
+	}
+
+	void FEventCollectionEditor::TransactEventCollection(FText Description, TFunctionRef<void(FSubsonicEventCollectionDefinition&)> TransactionFunc) const
+	{
+		if (USubsonicEventCollection* Collection = GetEventCollection())
+		{
+			FScopedTransaction Transaction(Description);
+			Collection->Modify();
+			TransactionFunc(Collection->GetMutableDefinition());
+		}
 	}
 
 	USubsonicEventCollection* FEventCollectionEditor::GetEventCollection() const
@@ -176,7 +887,7 @@ namespace UE::Subsonic::Editor
 	{
 		if (const ISlateStyle* Style = FSlateStyleRegistry::FindSlateStyle("SubsonicStyle"))
 		{
-			if (const FSlateBrush* Brush = Style->GetBrush("SubsonicEventCollectionEditor_Icon"))
+			if (const FSlateBrush* Brush = Style->GetBrush("SubsonicEditor.SubsonicEventCollection.Icon"))
 			{
 				return Brush;
 			}
@@ -201,6 +912,95 @@ namespace UE::Subsonic::Editor
 		return FAssetEditorToolkit::GetDefaultTabColor();
 	}
 
+	FSubsonicEventActionDefinition* FEventCollectionEditor::InsertAction(const FEventCollectionEditor::FHandleVariant& InHandleVariant)
+	{
+		FSubsonicEventActionDefinition* NewAction= nullptr;
+
+		const FText TransactionDesc = FText::Format(
+			InHandleVariant.IsEvent()
+			? LOCTEXT("AddActionToEndofEventToolTip_Format", "Add action to end of event '{0}' array")
+			: LOCTEXT("InsertActionInEventToolTip_Format", "Insert new action in event '{0}' array"),
+			FText::FromName(InHandleVariant.GetEvent().EventName));
+
+		TransactEventCollection(TransactionDesc, [&](FSubsonicEventCollectionDefinition& Definition)
+		{
+			const FEventHandle& EventHandle = InHandleVariant.GetEvent();
+			int32 InsertIndex = InHandleVariant.IsEvent() ? INDEX_NONE : InHandleVariant.GetAction().Index;
+			NewAction = Definition.AddAction(EventHandle, InsertIndex);
+		});
+
+		return NewAction;
+	}
+
+	void FEventCollectionEditor::BuildTransport()
+	{
+		TSharedRef<FExtender> ToolbarExtender = MakeShared<FExtender>();
+		ToolbarExtender->AddToolBarExtension(
+			"Asset",
+			EExtensionHook::After,
+			GetToolkitCommands(),
+			FToolBarExtensionDelegate::CreateLambda([this](FToolBarBuilder& ToolbarBuilder)
+			{
+				const FEventCollectionEditorCommands& Commands = FEventCollectionEditorCommands::Get();
+
+				ToolbarBuilder.BeginSection("Transport");
+				{
+					ToolbarBuilder.BeginStyleOverride("Toolbar.BackplateLeft");
+					{
+						ToolbarBuilder.AddToolBarButton(Commands.StartAudition, NAME_None, { }, { },
+							TAttribute<FSlateIcon>::Create([this]()
+							{
+								if (!PreviewExecutor.IsValid())
+								{
+									static const FSlateIcon IconInactive = FSlateIcon("SubsonicStyle", "SubsonicEditor.Play.Inactive.Valid");
+									return IconInactive;
+								}
+
+								static const FSlateIcon IconActive = FSlateIcon("SubsonicStyle", "SubsonicEditor.Play.Active.Valid");
+								return IconActive;
+							})
+						);
+					}
+					ToolbarBuilder.EndStyleOverride();
+
+					ToolbarBuilder.BeginStyleOverride("Toolbar.BackplateRight");
+					{
+						ToolbarBuilder.AddToolBarButton(Commands.StopAudition, NAME_None, { }, { },
+							TAttribute<FSlateIcon>::Create([this]()
+							{
+								if (PreviewExecutor.IsValid())
+								{
+									static const FSlateIcon IconActive = FSlateIcon("SubsonicStyle", "SubsonicEditor.Stop.Active");
+									return IconActive;
+								}
+
+								static const FSlateIcon IconInactive = FSlateIcon("SubsonicStyle", "SubsonicEditor.Stop.Inactive");
+								return IconInactive;
+							})
+						);
+					}
+					ToolbarBuilder.EndStyleOverride();
+				}
+				ToolbarBuilder.EndSection();
+			})
+		);
+
+		AddToolbarExtender(ToolbarExtender);
+	}
+
+	void FEventCollectionEditor::ExecutePreviewEvent(const FEventHandle& EventHandle)
+	{
+		if (!PreviewExecutor.IsValid())
+		{
+			RestartAudition(false /* bAutoEventsEnabled */);
+		}
+
+		if (PreviewExecutor.IsValid())
+		{
+			PreviewExecutor->GetExecutor().ExecuteEvent(EventHandle.EventName);
+		}
+	}
+
 	FName FEventCollectionEditor::GetEditorName() const
 	{
 		return FEventCollectionEditor::EditorName;
@@ -209,7 +1009,71 @@ namespace UE::Subsonic::Editor
 	void FEventCollectionEditor::PostUndo(bool bSuccess)
 	{
 		FSlateApplication::Get().DismissAllMenus();
+		RebuildEventTree(true /* bRebuildListView */);
 		FEditorUndoClient::PostUndo(bSuccess);
 	}
+
+	void FEventCollectionEditor::RestartAudition(bool bAutoEventsEnabled)
+	{
+		if (const USubsonicEventCollection* EventCollection = GetEventCollection())
+		{
+			if (FAudioDeviceManager* DeviceManager = GEngine->GetAudioDeviceManager())
+			{
+				if (FAudioDeviceHandle ActiveDevice = DeviceManager->GetActiveAudioDevice(); ActiveDevice.IsValid())
+				{
+					if (PreviewExecutor.IsValid())
+					{
+						PreviewExecutor->Unregister();
+					}
+
+					UObject* TransientObject = GetTransientPackageAsObject();
+					check(TransientObject);
+					PreviewExecutor.Reset(USubsonicEventCollectionExecutor::Create(*TransientObject, "PreviewExecutor", *EventCollection, ActiveDevice.GetDeviceID()));
+
+					if (bAutoEventsEnabled)
+					{
+						const FSubsonicEventCollectionDefinition& CollectionDefinition = EventCollection->GetDefinition();
+						for (const TPair<FGameplayTag, FSubsonicEvent>& Pair : CollectionDefinition.GetEvents())
+						{
+							if (Pair.Value.GetAutoAudition())
+							{
+								PreviewExecutor->GetExecutor().ExecuteEvent(Pair.Key.GetTagName());
+							}
+						}
+					}
+				}
+			}
+		}
+	}
+
+	void FEventCollectionEditor::StartAudition()
+	{
+		if (!PreviewExecutor.IsValid())
+		{
+			RestartAudition();
+		}
+	}
+
+	void FEventCollectionEditor::StopAudition()
+	{
+		if (PreviewExecutor.IsValid())
+		{
+			PreviewExecutor->Unregister();
+			PreviewExecutor.Reset();
+		}
+	}
+
+	void FEventCollectionEditor::ToggleAudition()
+	{
+		if (PreviewExecutor.IsValid())
+		{
+			PreviewExecutor->Unregister();
+			PreviewExecutor.Reset();
+		}
+		else
+		{
+			RestartAudition();
+		}
+	}
 } // namespace UE::Subsonic::Editor
 #undef LOCTEXT_NAMESPACE // SubsonicEditor
diff --git a/Engine/Plugins/Experimental/Subsonic/Source/SubsonicEditor/Private/SubsonicEventCollectionEditor.h b/Engine/Plugins/Experimental/Subsonic/Source/SubsonicEditor/Private/SubsonicEventCollectionEditor.h
--- a/Engine/Plugins/Experimental/Subsonic/Source/SubsonicEditor/Private/SubsonicEventCollectionEditor.h
+++ b/Engine/Plugins/Experimental/Subsonic/Source/SubsonicEditor/Private/SubsonicEventCollectionEditor.h
@@ -7,11 +7,17 @@
 #include "Framework/Docking/TabManager.h"
 #include "IDetailsView.h"
 #include "Misc/NotifyHook.h"
+#include "Misc/TVariant.h"
+#include "SubsonicEventCollection.h"
 #include "SubsonicEventCollectionViews.h"
+#include "SubsonicHandles.h"
+#include "Templates/Function.h"
 #include "Templates/SharedPointer.h"
 #include "Toolkits/AssetEditorToolkit.h"
 #include "Toolkits/IToolkitHost.h"
 #include "UObject/StrongObjectPtrTemplates.h"
+#include "Widgets/Views/SListView.h"
+
 
 namespace UE::Subsonic
 {
@@ -35,7 +41,7 @@ namespace UE::Subsonic
 
 			USubsonicEventCollection* GetEventCollection() const;
 
-			/** IToolkit interface */
+			/** IToolkit Interface */
 			virtual FName GetToolkitFName() const override;
 			virtual FText GetBaseToolkitName() const override;
 			virtual FString GetWorldCentricTabPrefix() const override;
@@ -43,7 +49,7 @@ namespace UE::Subsonic
 			virtual const FSlateBrush* GetDefaultTabIcon() const override;
 			virtual FLinearColor GetDefaultTabColor() const override;
 
-			/** IAssetEditorInstance interface */
+			/** IAssetEditorInstance Interface */
 			virtual FName GetEditorName() const override;
 
 			/** FEditorUndoClient Interface */
@@ -51,13 +57,86 @@ namespace UE::Subsonic
 			virtual void PostRedo(bool bSuccess) override { PostUndo(bSuccess); }
 
 		private:
-			TSharedPtr<IDetailsView> EventsDetailsView;
-			TSharedPtr<FTabManager::FStack> EventsStack;
-			TStrongObjectPtr<USubsonicEventsView> EventsView;
+			void TransactEvent(const FEventHandle EventHandle, FText Description, TFunctionRef<void(FSubsonicEvent&)> TransactionFunc) const;
+			void TransactEventCollection(FText Description, TFunctionRef<void(FSubsonicEventCollectionDefinition&)> TransactionFunc) const;
+
+			TSharedRef<SWidget> InitEventTreeView();
+
+			class FHandleVariant
+			{
+				TVariant<FEventHandle, FActionHandle> Handle;
+
+			public:
+				const FActionHandle& GetAction() const
+				{
+					checkf(IsAction(), TEXT("Must check that handle is action first"));
+					return Handle.Get<FActionHandle>();
+				}
+
+				const FEventHandle& GetEvent() const
+				{
+					return Handle.IsType<FActionHandle>()
+						? Handle.Get<FActionHandle>().Event
+						: Handle.Get<FEventHandle>();
+				}
+
+				bool IsAction() const
+				{
+					return Handle.IsType<FActionHandle>();
+				}
+
+				bool IsEvent() const
+				{
+					return Handle.IsType<FEventHandle>();
+				}
+
+				template <typename HandleType>
+				void Set(HandleType NewHandle)
+				{
+					Handle.Set<HandleType>(MoveTemp(NewHandle));
+				}
+			};
+
+			// Inserts action. If the handle variant points at an event, adds to end of event array.
+			// If it points to an action, inserts before the given action.
+			FSubsonicEventActionDefinition* InsertAction(const FHandleVariant& InHandleVariant);
+
+			void BuildTransport();
+
+			void ExecutePreviewEvent(const FEventHandle& EventHandle);
+
+			TSharedRef<ITableRow> OnEventTree_GenerateRow(TSharedPtr<FHandleVariant> HandleVariant, const TSharedRef<STableViewBase>& TableView);
+			void OnEventTree_RowSelected(TSharedPtr<FHandleVariant> HandleVariant, ESelectInfo::Type SelectInfo);
+
+			void RebuildEventTree(bool bRebuildListView = false);
+
+			void RestartAudition(bool bAutoEventsEnabled = true);
+			void StartAudition();
+			void StopAudition();
+			void ToggleAudition();
+
+			TStrongObjectPtr<USubsonicEventCollectionExecutor> PreviewExecutor;
+
+			TSharedPtr<IDetailsView> EventTreeInspectorView;
+			TSharedPtr<FTabManager::FStack> DetailsStack;
+			TStrongObjectPtr<USubsonicEventTreeDetailsView> TreeDetailsView;
+
+			TSharedPtr<FTabManager::FStack> EventTreeStack;
 
 			TSharedPtr<IDetailsView> ParametersDetailsView;
 			TSharedPtr<FTabManager::FStack> ParametersStack;
-			TStrongObjectPtr<USubsonicParametersView> ParametersView;
+			TStrongObjectPtr<USubsonicCollectionParametersView> ParametersView;
+
+			TArray<TSharedPtr<FHandleVariant>> EventTreeHandles;
+			TSharedPtr<SListView<TSharedPtr<FHandleVariant>>> EventTreeView;
+			FString EventViewFilter;
+
+			// Value names of selectable actions in event tree view
+			TArray<TSharedPtr<FString>> ActionComboValues;
+
+			// Single property view of collection for use to capture bespoke property handles
+			// for visualization in contexts outside of a contextual details layout/panel
+			TSharedPtr<ISinglePropertyView> CollectionView;
 		};
 	} // namespace Editor
 } // namespace UE::Subsonic
\ No newline at end of file
diff --git a/Engine/Plugins/Experimental/Subsonic/Source/SubsonicEditor/Private/SubsonicEventCollectionEditorCommands.cpp b/Engine/Plugins/Experimental/Subsonic/Source/SubsonicEditor/Private/SubsonicEventCollectionEditorCommands.cpp
new file mode 100644
--- /dev/null
+++ b/Engine/Plugins/Experimental/Subsonic/Source/SubsonicEditor/Private/SubsonicEventCollectionEditorCommands.cpp
@@ -0,0 +1,18 @@
+// Copyright Epic Games, Inc. All Rights Reserved.
+#include "SubsonicEventCollectionEditorCommands.h"
+
+#include "Internationalization/Internationalization.h"
+#include "Styling/AppStyle.h"
+
+
+#define LOCTEXT_NAMESPACE "SubsonicEditor"
+namespace UE::Subsonic::Editor
+{
+	void FEventCollectionEditorCommands::RegisterCommands()
+	{
+		UI_COMMAND(StartAudition, "Start Audition", "Creates and initalizes Subsonic audition executor", EUserInterfaceActionType::Button, FInputChord());
+		UI_COMMAND(StopAudition, "Stop Audition", "Destroys the audition executor, clearing associate subscriber state", EUserInterfaceActionType::Button, FInputChord());
+		UI_COMMAND(ToggleAudition, "Toggle Audition", "Starts or stops the currently auditioning executor", EUserInterfaceActionType::Button, FInputChord(EKeys::SpaceBar));
+	}
+} // namespace UE::Subsonic::Editor
+#undef LOCTEXT_NAMESPACE // SubsonicEditor
diff --git a/Engine/Plugins/Experimental/Subsonic/Source/SubsonicEditor/Private/SubsonicEventCollectionEditorCommands.h b/Engine/Plugins/Experimental/Subsonic/Source/SubsonicEditor/Private/SubsonicEventCollectionEditorCommands.h
new file mode 100644
--- /dev/null
+++ b/Engine/Plugins/Experimental/Subsonic/Source/SubsonicEditor/Private/SubsonicEventCollectionEditorCommands.h
@@ -0,0 +1,29 @@
+// Copyright Epic Games, Inc. All Rights Reserved.
+
+#pragma once
+
+#include "Framework/Commands/Commands.h"
+
+namespace UE::Subsonic::Editor
+{
+	class FEventCollectionEditorCommands : public TCommands<FEventCollectionEditorCommands>
+	{
+	public:
+		FEventCollectionEditorCommands()
+			: TCommands<FEventCollectionEditorCommands>("SubsonicEditor", NSLOCTEXT("SubsonicEditor", "SubsonicEditorCommands_EditorName", "Subsonic Event Collection Editor"), NAME_None, "SubsonicStyle")
+		{
+		}
+
+		/** Creates and initializes Subsonic audition executor */
+		TSharedPtr<FUICommandInfo> StartAudition;
+
+		/** Destroys the auditioning executor, clearing associate subscriber state. */
+		TSharedPtr<FUICommandInfo> StopAudition;
+
+		/** Starts or stops the currently auditioning executor */
+		TSharedPtr<FUICommandInfo> ToggleAudition;
+
+		/** Initialize commands */
+		virtual void RegisterCommands() override;
+	};
+} // namespace UE::Subsonic::Editor
\ No newline at end of file
diff --git a/Engine/Plugins/Experimental/Subsonic/Source/SubsonicEditor/Private/SubsonicEventCollectionViews.cpp b/Engine/Plugins/Experimental/Subsonic/Source/SubsonicEditor/Private/SubsonicEventCollectionViews.cpp
--- a/Engine/Plugins/Experimental/Subsonic/Source/SubsonicEditor/Private/SubsonicEventCollectionViews.cpp
+++ b/Engine/Plugins/Experimental/Subsonic/Source/SubsonicEditor/Private/SubsonicEventCollectionViews.cpp
@@ -3,10 +3,7 @@
 #include "SubsonicEventCollectionViews.h"
 
 #include "DetailCategoryBuilder.h"
-#include "DetailWidgetRow.h"
-#include "IDetailChildrenBuilder.h"
 #include "Internationalization/Internationalization.h"
-#include "Misc/AssertionMacros.h"
 #include "SubsonicEventCollection.h"
 
 #include UE_INLINE_GENERATED_CPP_BY_NAME(SubsonicEventCollectionViews)
@@ -16,44 +13,173 @@
 
 namespace UE::Subsonic::Editor
 {
-	void FCollectionPropertyCustomizationBase::CustomizeDetails(IDetailLayoutBuilder& DetailLayout)
+	namespace CollectionViewsPrivate
 	{
-		// Only supports modifying a single collection at a time
-		TArray<TWeakObjectPtr<UObject>> Objects;
-		DetailLayout.GetObjectsBeingCustomized(Objects);
-		if (Objects.Num() == 1 && Objects.Last().IsValid())
+		void AddActionProperties(IDetailLayoutBuilder& LayoutBuilder, IDetailCategoryBuilder& CategoryBuilder, TSharedRef<IPropertyHandleArray> ActionArrayHandle, uint32 ActionIndex)
 		{
-			if (USubsonicEventCollectionViewBase* View = CastChecked<USubsonicEventCollectionViewBase>(Objects.Last()))
+			TSharedRef<IPropertyHandle> ActionEntryHandle = ActionArrayHandle->GetElement(ActionIndex);
+			uint32 NumProperties = INDEX_NONE;
+			if (ActionEntryHandle->GetNumChildren(NumProperties) == FPropertyAccess::Success)
 			{
-				TSharedPtr<IPropertyHandle> DefinitionHandle = DetailLayout.AddObjectPropertyData(
-					TArray<UObject*> { View->GetCollection() },
-					USubsonicEventCollection::GetDefinitionPropertyName());
-
-				if (DefinitionHandle.IsValid())
+				for (uint32 Index = 0; Index < NumProperties; ++Index)
 				{
-					const FText CategoryDisplayName = HeaderName.IsEmptyOrWhitespace() ? FText::FromName(PropertyName) : HeaderName;
-					TSharedPtr<IPropertyHandle> ChildProperty = DefinitionHandle->GetChildHandle(PropertyName);
-					if (ChildProperty.IsValid())
+					TSharedPtr<IPropertyHandle> ChildHandle = ActionEntryHandle->GetChildHandle(Index);
+					if (ChildHandle.IsValid())
 					{
-						IDetailCategoryBuilder& Category = DetailLayout.EditCategory(PropertyName, CategoryDisplayName);
-						Category.InitiallyCollapsed(false);
-						Category.AddProperty(ChildProperty);
+						const FName Category = ChildHandle->GetDefaultCategoryName();
+						Category.IsNone() ? CategoryBuilder.AddProperty(ChildHandle) : LayoutBuilder.EditCategory(Category).AddProperty(ChildHandle);
 					}
 				}
 			}
 		}
+
+		TSharedPtr<IPropertyHandle> GetEventHandleMatchingName(TSharedRef<IPropertyHandleMap> EventMapHandle, uint32 Index, FName Name)
+		{
+			TSharedRef<IPropertyHandle> EventEntryHandle = EventMapHandle->GetElement(Index);
+			if (TSharedPtr<IPropertyHandle> EventKeyHandle = EventEntryHandle->GetKeyHandle(); EventKeyHandle.IsValid())
+			{
+				FName EventName;
+				TSharedPtr<IPropertyHandle> TagNameHandle = EventKeyHandle->GetChildHandle("TagName");
+				if (TagNameHandle.IsValid()
+					&& TagNameHandle->GetValue(EventName) == FPropertyAccess::Success
+					&& EventName == Name)
+				{
+					return EventEntryHandle;
+				}
+			}
+
+			return { };
+		}
+	} // ViewsPrivate
+
+	void FCollectionParametersDetailCustomization::CustomizeDetails(IDetailLayoutBuilder& LayoutBuilder)
+	{
+		if (USubsonicEventCollectionViewBase* View = GetViewFromBuilder(LayoutBuilder))
+		{
+			const FName ParametersPropertyName = FSubsonicEventCollectionDefinition::GetParametersPropertyName();
+			TSharedPtr<IPropertyHandle> ChildHandle = CacheCollectionPropertyHandle(LayoutBuilder, *View, ParametersPropertyName);
+			if (ChildHandle.IsValid())
+			{
+				const FText CategoryDisplayName = HeaderName.IsEmptyOrWhitespace() ? FText::FromName(ParametersPropertyName) : HeaderName;
+				IDetailCategoryBuilder& CategoryBuilder = LayoutBuilder.EditCategory(ParametersPropertyName, CategoryDisplayName);
+				CategoryBuilder.InitiallyCollapsed(false);
+
+				const FName Category = ChildHandle->GetDefaultCategoryName();
+				Category.IsNone() ? CategoryBuilder.AddProperty(ChildHandle) : LayoutBuilder.EditCategory(Category).AddProperty(ChildHandle);
+			}
+		}
+	}
+
+	TSharedPtr<IPropertyHandle> FCollectionPropertyCustomizationBase::CacheCollectionPropertyHandle(
+		IDetailLayoutBuilder& LayoutBuilder,
+		USubsonicEventCollectionViewBase& View,
+		FName ChildPropertyName)
+	{
+		if (USubsonicEventCollection* ViewCollection = View.GetCollection())
+		{
+			TSharedPtr<IPropertyHandle> DefinitionHandle = LayoutBuilder.AddObjectPropertyData(
+				TArray<UObject*> { ViewCollection },
+				USubsonicEventCollection::GetDefinitionPropertyName());
+
+			if (DefinitionHandle.IsValid())
+			{
+				return DefinitionHandle->GetChildHandle(ChildPropertyName);
+			}
+		}
+
+		return { };
+	}
+
+	USubsonicEventCollectionViewBase* FCollectionPropertyCustomizationBase::GetViewFromBuilder(IDetailLayoutBuilder& LayoutBuilder)
+	{
+		// Only supports modifying a single collection view at a time
+		TArray<TWeakObjectPtr<UObject>> Objects;
+		LayoutBuilder.GetObjectsBeingCustomized(Objects);
+		if (Objects.Num() == 1 && Objects.Last().IsValid())
+		{
+			return CastChecked<USubsonicEventCollectionViewBase>(Objects.Last());
+		}
+
+		return nullptr;
 	}
 
 	FCollectionParametersDetailCustomization::FCollectionParametersDetailCustomization()
 	{
-		PropertyName = FSubsonicEventCollectionDefinition::GetParametersPropertyName();
 		HeaderName = LOCTEXT("ParametersDetails_DisplayName", "Parameters");
 	}
 
-	FCollectionEventsDetailCustomization::FCollectionEventsDetailCustomization()
+	FEventTreeSelectionDetailCustomization::FEventTreeSelectionDetailCustomization()
 	{
-		PropertyName = FSubsonicEventCollectionDefinition::GetEventsPropertyName();
-		HeaderName = LOCTEXT("EventsDetails_DisplayName", "Events");
+		HeaderName = LOCTEXT("EventTreeDetails_DisplayName", "Details");
+	}
+
+	void FEventTreeSelectionDetailCustomization::CustomizeDetails(IDetailLayoutBuilder& LayoutBuilder)
+	{
+		USubsonicEventTreeDetailsView* View = Cast<USubsonicEventTreeDetailsView>(GetViewFromBuilder(LayoutBuilder));
+		if (!View)
+		{
+			return;
+		}
+
+		TSharedPtr<IPropertyHandle> ChildProperty = CacheCollectionPropertyHandle(LayoutBuilder, *View, FSubsonicEventCollectionDefinition::GetEventsPropertyName());
+		if (!ChildProperty.IsValid())
+		{
+			return;
+		}
+
+		TSharedPtr<IPropertyHandleMap> EventMapHandle = ChildProperty->AsMap();
+		if (!EventMapHandle.IsValid())
+		{
+			return;
+		}
+
+		uint32 NumEvents = 0;
+		if (EventMapHandle->GetNumElements(NumEvents) == FPropertyAccess::Success && NumEvents > 0)
+		{
+			const FText CategoryDisplayName = HeaderName.IsEmptyOrWhitespace() ? FText::FromName(FSubsonicEventCollectionDefinition::GetEventsPropertyName()) : HeaderName;
+			IDetailCategoryBuilder& CategoryBuilder = LayoutBuilder.EditCategory(ChildProperty->GetDefaultCategoryName(), CategoryDisplayName);
+			CategoryBuilder.InitiallyCollapsed(false);
+
+			for (uint32 EventIndex = 0; EventIndex < NumEvents; ++EventIndex)
+			{
+				const FName ActionCollectionPropertyName = FSubsonicEvent::GetActionCollectionPropertyName();
+				if (TSharedPtr<IPropertyHandle> EventHandle = CollectionViewsPrivate::GetEventHandleMatchingName(EventMapHandle.ToSharedRef(), EventIndex, View->Event))
+				{
+					// Show Event Properties (View->ActionIndex is INDEX_NONE/-1)
+					if (View->ActionIndex < 0)
+					{
+						uint32 NumChildren = 0;
+						EventHandle->GetNumChildren(NumChildren);
+						for (uint32 ChildIndex = 0; ChildIndex < NumChildren; ++ChildIndex)
+						{
+							if (TSharedPtr<IPropertyHandle> ChildHandle = EventHandle->GetChildHandle(ChildIndex))
+							{
+								if (FProperty* Property = ChildHandle->GetProperty(); Property && Property->GetFName() != ActionCollectionPropertyName)
+								{
+									const FName Category = ChildHandle->GetDefaultCategoryName();
+									Category.IsNone() ? CategoryBuilder.AddProperty(ChildHandle) : LayoutBuilder.EditCategory(Category).AddProperty(ChildHandle);
+								}
+							}
+						}
+					}
+					// Show Action (Selected) Properties
+					else
+					{
+						if (TSharedPtr<IPropertyHandle> ActionCollectionHandle = EventHandle->GetChildHandle(ActionCollectionPropertyName, false /* bRecurse */))
+						{
+							if (TSharedPtr<IPropertyHandleArray> ArrayHandle = ActionCollectionHandle->AsArray())
+							{
+								uint32 NumActions = 0;
+								if (ArrayHandle->GetNumElements(NumActions) == FPropertyAccess::Success && static_cast<int32>(NumActions) > View->ActionIndex)
+								{
+									CollectionViewsPrivate::AddActionProperties(LayoutBuilder, CategoryBuilder, ArrayHandle.ToSharedRef(), View->ActionIndex);
+								}
+							}
+						}
+					}
+				}
+			}
+		}
 	}
 } // namespace UE::Subsonic::Editor
 #undef LOCTEXT_NAMESPACE // "SubsonicEditor"
diff --git a/Engine/Plugins/Experimental/Subsonic/Source/SubsonicEditor/Private/SubsonicEventCollectionViews.h b/Engine/Plugins/Experimental/Subsonic/Source/SubsonicEditor/Private/SubsonicEventCollectionViews.h
--- a/Engine/Plugins/Experimental/Subsonic/Source/SubsonicEditor/Private/SubsonicEventCollectionViews.h
+++ b/Engine/Plugins/Experimental/Subsonic/Source/SubsonicEditor/Private/SubsonicEventCollectionViews.h
@@ -27,38 +27,50 @@ namespace UE::Subsonic::Editor
 	};
 
 	UCLASS(MinimalAPI)
-	class USubsonicParametersView : public USubsonicEventCollectionViewBase
+	class USubsonicCollectionParametersView : public USubsonicEventCollectionViewBase
 	{
 		GENERATED_BODY()
 	};
 
 	UCLASS(MinimalAPI)
-	class USubsonicEventsView : public USubsonicEventCollectionViewBase
+	class USubsonicEventTreeDetailsView : public USubsonicEventCollectionViewBase
 	{
 		GENERATED_BODY()
+
+	public:
+		FName Event;
+		int32 ActionIndex = INDEX_NONE;
 	};
 
 	class FCollectionPropertyCustomizationBase : public IDetailCustomization
 	{
 	protected:
-		FName PropertyName;
 		FText HeaderName;
 
-	public:
-		// IDetailCustomization interface
-		virtual void CustomizeDetails(IDetailLayoutBuilder& DetailLayout) override;
-		// End of IDetailCustomization interface
+	protected:
+		static TSharedPtr<IPropertyHandle> CacheCollectionPropertyHandle(
+			IDetailLayoutBuilder& LayoutBuilder,
+			USubsonicEventCollectionViewBase& View,
+			FName ChildPropertyName);
+
+		static USubsonicEventCollectionViewBase* GetViewFromBuilder(IDetailLayoutBuilder& LayoutBuilder);
 	};
 
 	class FCollectionParametersDetailCustomization : public FCollectionPropertyCustomizationBase
 	{
 	public:
 		FCollectionParametersDetailCustomization();
+
+		// IDetailCustomization interface
+		virtual void CustomizeDetails(IDetailLayoutBuilder& DetailLayout) override;
+		// End of IDetailCustomization interface
 	};
 
-	class FCollectionEventsDetailCustomization : public FCollectionPropertyCustomizationBase
+	class FEventTreeSelectionDetailCustomization : public FCollectionPropertyCustomizationBase
 	{
 	public:
-		FCollectionEventsDetailCustomization();
+		FEventTreeSelectionDetailCustomization();
+
+		virtual void CustomizeDetails(IDetailLayoutBuilder& LayoutBuilder) override;
 	};
 } // namespace UE::Subsonic::Editor
\ No newline at end of file
diff --git a/Engine/Plugins/Experimental/Subsonic/Source/SubsonicEditor/Private/SubsonicSlateStyle.h b/Engine/Plugins/Experimental/Subsonic/Source/SubsonicEditor/Private/SubsonicSlateStyle.h
--- a/Engine/Plugins/Experimental/Subsonic/Source/SubsonicEditor/Private/SubsonicSlateStyle.h
+++ b/Engine/Plugins/Experimental/Subsonic/Source/SubsonicEditor/Private/SubsonicSlateStyle.h
@@ -2,6 +2,7 @@
 #pragma once
 
 #include "AudioWidgetsStyle.h"
+#include "Styling/SlateStyleMacros.h"
 #include "Styling/SlateStyleRegistry.h"
 
 
@@ -15,20 +16,38 @@ namespace UE::Subsonic::Editor
 		{
 			SetParentStyleName(FAudioWidgetsStyle::Get().GetStyleSetName());
 
-			SetContentRoot(FPaths::EnginePluginsDir() / TEXT("Runtime/Subsonic/Content/Editor/Slate"));
+			SetContentRoot(FPaths::EnginePluginsDir() / TEXT("Experimental/Subsonic/Content/Editor/Slate"));
 			SetCoreContentRoot(FPaths::EngineContentDir() / TEXT("Slate"));
 
-			static const FVector2D Icon20x20(20.0f, 20.0f);
-			static const FVector2D Icon40x40(40.0f, 40.0f);
-
-			static const FVector2D Icon16 = FVector2D(16.0f, 16.0f);
+			static const FVector2D Icon20(20.0f, 20.0f);
+			static const FVector2D Icon40(40.0f, 40.0f);
 			static const FVector2D Icon64 = FVector2D(64.0f, 64.0f);
 
-			const FVector2D Icon15x11(15.0f, 11.0f);
-
 			// Subsonic Editor
 			{
 				Set("SubsonicEventCollection.Color", FColor(156, 0, 255));
+
+				Set("SubsonicEditor.Play", new IMAGE_BRUSH_SVG(TEXT("Icons/play"), Icon40));
+				Set("SubsonicEditor.Play.Small", new IMAGE_BRUSH_SVG(TEXT("Icons/play"), Icon20));
+				Set("SubsonicEditor.Play.Thumbnail", new IMAGE_BRUSH_SVG(TEXT("Icons/play_thumbnail"), Icon64));
+				Set("SubsonicEditor.Play.Thumbnail.Hovered", new IMAGE_BRUSH_SVG(TEXT("Icons/play_thumbnail_hover"), Icon64));
+
+				Set("SubsonicEditor.Play.Active.Valid", new IMAGE_BRUSH_SVG(TEXT("Icons/play_active_valid"), Icon40));
+				Set("SubsonicEditor.Play.Active.Warning", new IMAGE_BRUSH_SVG(TEXT("Icons/play_active_warning"), Icon40));
+				Set("SubsonicEditor.Play.Inactive.Valid", new IMAGE_BRUSH_SVG(TEXT("Icons/play_inactive_valid"), Icon40));
+				Set("SubsonicEditor.Play.Inactive.Warning", new IMAGE_BRUSH_SVG(TEXT("Icons/play_inactive_warning"), Icon40));
+				Set("SubsonicEditor.Play.Error", new IMAGE_BRUSH_SVG(TEXT("Icons/play_error"), Icon40));
+
+				Set("SubsonicEditor.Stop", new IMAGE_BRUSH_SVG(TEXT("Icons/stop"), Icon40));
+
+				Set("SubsonicEditor.Stop.Disabled", new IMAGE_BRUSH_SVG(TEXT("Icons/stop_disabled"), Icon40));
+				Set("SubsonicEditor.Stop.Active", new IMAGE_BRUSH_SVG(TEXT("Icons/stop_active"), Icon40));
+				Set("SubsonicEditor.Stop.Inactive", new IMAGE_BRUSH_SVG(TEXT("Icons/stop_inactive"), Icon40));
+				Set("SubsonicEditor.Stop.Thumbnail", new IMAGE_BRUSH_SVG(TEXT("Icons/stop_thumbnail"), Icon64));
+				Set("SubsonicEditor.Stop.Thumbnail.Hovered", new IMAGE_BRUSH_SVG(TEXT("Icons/stop_thumbnail_hover"), Icon64));
+
+				Set("SubsonicEditor.SubsonicEventCollection.Icon", new IMAGE_BRUSH_SVG(TEXT("Icons/subsoniceventcollection_icon"), Icon20));
+				Set("SubsonicEditor.SubsonicEventCollection.Thumbnail", new IMAGE_BRUSH_SVG(TEXT("Icons/subsoniceventcollection_thumbnail"), Icon20));
 			}
 
 			FSlateStyleRegistry::RegisterSlateStyle(*this);
diff --git a/Engine/Plugins/Experimental/Subsonic/Source/SubsonicEditor/SubsonicEditor.Build.cs b/Engine/Plugins/Experimental/Subsonic/Source/SubsonicEditor/SubsonicEditor.Build.cs
--- a/Engine/Plugins/Experimental/Subsonic/Source/SubsonicEditor/SubsonicEditor.Build.cs
+++ b/Engine/Plugins/Experimental/Subsonic/Source/SubsonicEditor/SubsonicEditor.Build.cs
@@ -7,7 +7,7 @@ public class SubsonicEditor : ModuleRules
 	public SubsonicEditor(ReadOnlyTargetRules Target) : base(Target)
 	{
 		PCHUsage = ModuleRules.PCHUsageMode.UseExplicitOrSharedPCHs;
-		
+
 		PublicDependencyModuleNames.AddRange(
 			new string[]
 			{
@@ -23,10 +23,18 @@ public class SubsonicEditor : ModuleRules
 				"AudioEditor",
 				"AudioWidgets",
 				"ContentBrowser",
+				"EditorSubsystem",
+				"Engine",
+				"GameplayTags",
+				"GameplayTagsEditor",
+				"InputCore",
+				"PropertyBindingUtils",
+				"PropertyBindingUtilsEditor",
 				"PropertyEditor",
 				"Slate",
 				"SlateCore",
 				"Subsonic",
+				"ToolWidgets",
 				"UnrealEd"
 			}
 		);
diff --git a/Engine/Plugins/Experimental/Subsonic/Subsonic.uplugin b/Engine/Plugins/Experimental/Subsonic/Subsonic.uplugin
--- a/Engine/Plugins/Experimental/Subsonic/Subsonic.uplugin
+++ b/Engine/Plugins/Experimental/Subsonic/Subsonic.uplugin
@@ -32,6 +32,10 @@
 			"Name": "AudioWidgets",
 			"Enabled": true
 		},
+		{
+			"Name": "GameplayTagsEditor",
+			"Enabled": true
+		},
 		{
 			"Name": "PropertyBindingUtils",
 			"Enabled": true
