Drag-and-Drop for Subsonic event editor

#rnx

#Synced-CL 51567328

diff --git a/Engine/Plugins/Experimental/Subsonic/Source/SubsonicEditor/Private/SubsonicEventCollectionEditor.cpp b/Engine/Plugins/Experimental/Subsonic/Source/SubsonicEditor/Private/SubsonicEventCollectionEditor.cpp
--- a/Engine/Plugins/Experimental/Subsonic/Source/SubsonicEditor/Private/SubsonicEventCollectionEditor.cpp
+++ b/Engine/Plugins/Experimental/Subsonic/Source/SubsonicEditor/Private/SubsonicEventCollectionEditor.cpp
@@ -8,6 +8,7 @@
 #include "Framework/Application/SlateApplication.h"
 #include "Framework/Docking/TabManager.h"
 #include "GameplayTagContainer.h"
+#include "Input/DragAndDrop.h"
 #include "ISinglePropertyView.h"
 #include "Modules/ModuleManager.h"
 #include "PropertyBindingExtension.h"
@@ -70,6 +71,23 @@ namespace UE::Subsonic::Editor
 			return NullBrush;
 		}
 
+		class FActionDragDropOp : public FDragDropOperation
+		{
+		public:
+			DRAG_DROP_OPERATOR_TYPE(FActionDragDropOp, FDragDropOperation)
+
+			FActionHandle SourceAction;
+
+			static TSharedRef<FActionDragDropOp> New(FActionHandle InAction)
+			{
+				TSharedRef<FActionDragDropOp> Op = MakeShared<FActionDragDropOp>();
+				Op->SourceAction = MoveTemp(InAction);
+				Op->MouseCursor = EMouseCursor::GrabHandClosed;
+				Op->Construct();
+				return Op;
+			}
+		};
+
 		TSharedPtr<IPropertyHandle> GetEventTagPropertyHandle(TSharedRef<ISinglePropertyView> CollectionView, FName EventName)
 		{
 			TSharedPtr<IPropertyHandle> CollectionProperty = CollectionView->GetPropertyHandle();
@@ -118,19 +136,84 @@ namespace UE::Subsonic::Editor
 
 	const FName FEventCollectionEditor::EditorName = "SubsonicEventCollectionEditor";
 
+	TSharedRef<FEventCollectionEditor::SHandleVariantTableRow> FEventCollectionEditor::InitDraggableEventRow(TSharedPtr<FEventCollectionEditor::FHandleVariant> HandleVariant, const TSharedRef<STableViewBase>& TableView)
+	{
+		using namespace EventCollectionEditorPrivate;
+
+		if (!HandleVariant.IsValid())
+		{
+			return SNew(SHandleVariantTableRow, TableView);
+		}
+
+		return SNew(SHandleVariantTableRow, TableView)
+		.OnDragDetected_Lambda([this, HandleVariant](const FGeometry& Geometry, const FPointerEvent& MouseEvent)
+		{
+			if (const USubsonicEventCollection* Collection = GetEventCollection())
+			{
+				if (HandleVariant->IsAction())
+				{
+					FActionHandle ActionHandle = HandleVariant->GetAction();
+					if (const FSubsonicEventActionDefinition* Action = Collection->GetDefinition().FindAction(ActionHandle))
+					{
+						return FReply::Handled().BeginDragDrop(FActionDragDropOp::New(MoveTemp(ActionHandle)));
+					}
+				}
+			}
+			return FReply::Unhandled();
+		})
+		.OnCanAcceptDrop_Lambda([](const FDragDropEvent& DragDropEvent, EItemDropZone InitZone, TSharedPtr<FHandleVariant> TargetVariant)
+		{
+			TOptional<EItemDropZone> ToReturn;
+			if (TargetVariant.IsValid())
+			{
+				TSharedPtr<FDragDropOperation> Op = DragDropEvent.GetOperation();
+				if (Op.IsValid() && Op->IsOfType<FActionDragDropOp>())
+				{
+					// If event, add to end of event
+					if (TargetVariant->IsEvent())
+					{
+						ToReturn = EItemDropZone::OntoItem;
+					}
+					else
+					{
+						// Reject only a self-drop (same parent event, same index).  Cross-event drops are allowed.
+						const FActionHandle& SourceAction = StaticCastSharedRef<FActionDragDropOp>(Op.ToSharedRef())->SourceAction;
+						if (SourceAction != TargetVariant->GetAction())
+						{
+							ToReturn = InitZone;
+						}
+					}
+				}
+			}
+			return ToReturn;
+		})
+		.OnAcceptDrop_Lambda([this](const FDragDropEvent& DragDropEvent, EItemDropZone DropZone, TSharedPtr<FHandleVariant> HandleVariant)
+		{
+			TSharedPtr<FDragDropOperation> Op = DragDropEvent.GetOperation();
+			if (Op.IsValid() && Op->IsOfType<FActionDragDropOp>())
+			{
+				if (HandleVariant.IsValid())
+				{
+					const FActionHandle& SourceAction = StaticCastSharedRef<FActionDragDropOp>(Op.ToSharedRef())->SourceAction;
+					MoveDroppedAction(SourceAction, *HandleVariant.Get(), DropZone);
+				}
+			}
+			return FReply::Handled();
+		});
+	}
+
 	TSharedRef<ITableRow> FEventCollectionEditor::OnEventTree_GenerateRow(TSharedPtr<FEventCollectionEditor::FHandleVariant> HandleVariant, const TSharedRef<STableViewBase>& TableView)
 	{
-		using STableRowPtr = TSharedPtr<STableRow<TSharedPtr<FHandleVariant>>>;
-		STableRowPtr NewRow = SNew(STableRow<TSharedPtr<FHandleVariant>>, TableView);
+		TSharedRef<SHandleVariantTableRow> NewRow = InitDraggableEventRow(HandleVariant, TableView);
 
 		const USubsonicEventCollection* Collection = GetEventCollection();
 		if (!HandleVariant.IsValid() || !Collection)
 		{
-			return NewRow.ToSharedRef();
+			return NewRow;
 		}
 
 		constexpr float IndentMult = 4.f;
-		FText RowDisplayName = FText::FromName("None"); // Valid display name for actions with no struct type set
+		FText RowDisplayName = FText::FromName(NAME_None); // Valid display name for actions with no struct type set
 		float RowIndent = 5.f;
 		if (HandleVariant->IsAction())
 		{
@@ -449,7 +532,7 @@ namespace UE::Subsonic::Editor
 			]
 		];
 
-		return NewRow.ToSharedRef();
+		return NewRow;
 	}
 
 	void FEventCollectionEditor::OnEventTree_RowSelected(TSharedPtr<FHandleVariant> HandleVariant, ESelectInfo::Type SelectInfo)
@@ -815,7 +898,6 @@ namespace UE::Subsonic::Editor
 				SNew(SSearchBox)
 				.HintText(LOCTEXT("SearchDetailsHint", "Search"))
 				.OnTextChanged_Lambda(SetFilterString)
-// 				.OnTextCommitted_Lambda(SetFilterString)
 				.AddMetaData<FTagMetaData>(TEXT("Details.Search"))
 			]
 		]
@@ -827,6 +909,7 @@ namespace UE::Subsonic::Editor
 			.ListItemsSource(&EventTreeHandles)
 			.OnGenerateRow(this, &FEventCollectionEditor::OnEventTree_GenerateRow)
 			.OnSelectionChanged(this, &FEventCollectionEditor::OnEventTree_RowSelected)
+			.SelectionMode(ESelectionMode::Single)
 		];
 	}
 
@@ -932,6 +1015,79 @@ namespace UE::Subsonic::Editor
 		return NewAction;
 	}
 
+	void FEventCollectionEditor::MoveDroppedAction(const FActionHandle& SourceAction, const FHandleVariant& TargetHandle, EItemDropZone DropZone)
+	{
+		if (TargetHandle.IsAction() && SourceAction == TargetHandle.GetAction())
+		{
+			return;
+		}
+
+		TransactEventCollection(LOCTEXT("MoveActionDescription", "Move action in Subsonic Event Collection"), [&](FSubsonicEventCollectionDefinition& Definition)
+		{
+			FSubsonicEvent* SourceEvent = Definition.FindMutableEvent(SourceAction.Event);
+			if (!SourceEvent)
+			{
+				return;
+			}
+
+			TArray<FSubsonicEventActionDefinition>& SourceActions = SourceEvent->GetMutableActionCollection();
+			if (!SourceActions.IsValidIndex(SourceAction.Index))
+			{
+				return;
+			}
+
+			FSubsonicEvent* TargetEvent = Definition.FindMutableEvent(TargetHandle.GetEvent());
+			if (!TargetEvent)
+			{
+				return;
+			}
+
+			const int32 TargetIndex = TargetHandle.IsEvent()
+				? TargetEvent->GetActionCollection().Num()
+				: TargetHandle.GetAction().Index;
+
+			FSubsonicEventActionDefinition MovedAction = MoveTemp(SourceActions[SourceAction.Index]);
+			SourceActions.RemoveAt(SourceAction.Index);
+
+			int32 InsertAt = INDEX_NONE;
+			const int32 SourceIndex = SourceAction.Index;
+			switch (DropZone)
+			{
+				case EItemDropZone::AboveItem:
+				{
+					InsertAt = (TargetIndex > SourceIndex) ? TargetIndex - 1 : TargetIndex;
+				}
+				break;
+
+				case EItemDropZone::BelowItem:
+				{
+					InsertAt = (TargetIndex >= SourceIndex) ? TargetIndex : TargetIndex + 1;
+				}
+				break;
+
+				case EItemDropZone::OntoItem:
+				default:
+				{
+					InsertAt = TargetIndex;
+				}
+				break;
+			}
+
+			TArray<FSubsonicEventActionDefinition>& TargetActions = TargetEvent->GetMutableActionCollection();
+			if (InsertAt != INDEX_NONE && TargetActions.Num() >= InsertAt)
+			{
+				TargetActions.Insert(MoveTemp(MovedAction), InsertAt);
+			}
+			else
+			{
+				TargetActions.Add(MoveTemp(MovedAction));
+
+			}
+		});
+
+		RebuildEventTree(true /* bRebuildListView */);
+	}
+
 	void FEventCollectionEditor::BuildTransport()
 	{
 		TSharedRef<FExtender> ToolbarExtender = MakeShared<FExtender>();
diff --git a/Engine/Plugins/Experimental/Subsonic/Source/SubsonicEditor/Private/SubsonicEventCollectionEditor.h b/Engine/Plugins/Experimental/Subsonic/Source/SubsonicEditor/Private/SubsonicEventCollectionEditor.h
--- a/Engine/Plugins/Experimental/Subsonic/Source/SubsonicEditor/Private/SubsonicEventCollectionEditor.h
+++ b/Engine/Plugins/Experimental/Subsonic/Source/SubsonicEditor/Private/SubsonicEventCollectionEditor.h
@@ -17,6 +17,7 @@
 #include "Toolkits/IToolkitHost.h"
 #include "UObject/StrongObjectPtrTemplates.h"
 #include "Widgets/Views/SListView.h"
+#include "Widgets/Views/STableRow.h"
 
 
 namespace UE::Subsonic
@@ -101,10 +102,18 @@ namespace UE::Subsonic
 			// If it points to an action, inserts before the given action.
 			FSubsonicEventActionDefinition* InsertAction(const FHandleVariant& InHandleVariant);
 
+			// Moves an action to the position relative to TargetAction per the given DropZone.
+			void MoveDroppedAction(const FActionHandle& SourceAction, const FHandleVariant& HandleVariant, EItemDropZone DropZone);
+
 			void BuildTransport();
 
 			void ExecutePreviewEvent(const FEventHandle& EventHandle);
 
+			using SHandleVariantTableRow = STableRow<TSharedPtr<FHandleVariant>>;
+
+			// Currently, only supports actions. Could support dragging and ordering events in specialized view sorted by user-defined index in the future.
+			TSharedRef<SHandleVariantTableRow> InitDraggableEventRow(TSharedPtr<FHandleVariant> HandleVariant, const TSharedRef<STableViewBase>& TableView);
+
 			TSharedRef<ITableRow> OnEventTree_GenerateRow(TSharedPtr<FHandleVariant> HandleVariant, const TSharedRef<STableViewBase>& TableView);
 			void OnEventTree_RowSelected(TSharedPtr<FHandleVariant> HandleVariant, ESelectInfo::Type SelectInfo);
 
