PR #14570: Fix build with clang-cl, UE_API needs to be after UE_INTERNAL

#Synced-CL 52148387

diff --git a/Engine/Plugins/Editor/AssetManagerEditor/Source/AssetManagerEditor/Public/ReferenceViewer/EdGraphNode_Reference.h b/Engine/Plugins/Editor/AssetManagerEditor/Source/AssetManagerEditor/Public/ReferenceViewer/EdGraphNode_Reference.h
--- a/Engine/Plugins/Editor/AssetManagerEditor/Source/AssetManagerEditor/Public/ReferenceViewer/EdGraphNode_Reference.h
+++ b/Engine/Plugins/Editor/AssetManagerEditor/Source/AssetManagerEditor/Public/ReferenceViewer/EdGraphNode_Reference.h
@@ -60,7 +60,7 @@ class UEdGraphNode_Reference : public UEdGraphNode
 	UE_API UEdGraphPin* GetDependencyPin();
 	UE_API UEdGraphPin* GetReferencerPin();
 
-	UE_API UE_INTERNAL void UpdatePath();
+	UE_INTERNAL UE_API void UpdatePath();
 
 private:
 	UE_API void CacheAssetData(const FAssetData& AssetData);
diff --git a/Engine/Plugins/Editor/AssetManagerEditor/Source/AssetManagerEditor/Public/ReferenceViewer/EdGraph_ReferenceViewer.h b/Engine/Plugins/Editor/AssetManagerEditor/Source/AssetManagerEditor/Public/ReferenceViewer/EdGraph_ReferenceViewer.h
--- a/Engine/Plugins/Editor/AssetManagerEditor/Source/AssetManagerEditor/Public/ReferenceViewer/EdGraph_ReferenceViewer.h
+++ b/Engine/Plugins/Editor/AssetManagerEditor/Source/AssetManagerEditor/Public/ReferenceViewer/EdGraph_ReferenceViewer.h
@@ -99,8 +99,8 @@ public:
 	UE_API class UEdGraphNode_Reference* RefilterGraph();
 
 	UE_INTERNAL bool IsShowingContentVersePath() const { return bShowingContentVersePath; }
-	UE_API UE_INTERNAL void SetShowingContentVersePath(bool bInShowingContentVersePath);
-	UE_API UE_INTERNAL void UpdatePaths();
+	UE_INTERNAL UE_API void SetShowingContentVersePath(bool bInShowingContentVersePath);
+	UE_INTERNAL UE_API void UpdatePaths();
 
 	using FIsAssetIdentifierPassingSearchFilterCallback = TFunction<bool(const FAssetIdentifier&)>;
 	UE_DEPRECATED(5.7, "Call SetDoesAssetPassSearchFilterCallback instead.")
diff --git a/Engine/Plugins/Experimental/Subsonic/Source/Subsonic/Public/SubsonicEventCollection.h b/Engine/Plugins/Experimental/Subsonic/Source/Subsonic/Public/SubsonicEventCollection.h
--- a/Engine/Plugins/Experimental/Subsonic/Source/Subsonic/Public/SubsonicEventCollection.h
+++ b/Engine/Plugins/Experimental/Subsonic/Source/Subsonic/Public/SubsonicEventCollection.h
@@ -91,7 +91,7 @@ namespace UE::Subsonic
 		virtual ~FSubsonicEvent() = default;
 
 #if WITH_EDITOR
-		UE_API UE_INTERNAL static FName GetActionCollectionPropertyName();
+		UE_INTERNAL UE_API static FName GetActionCollectionPropertyName();
 #endif // WITH_EDITOR
 
 		UE_API const TArray<FSubsonicEventActionDefinition>& GetActionCollection() const;
@@ -103,8 +103,8 @@ namespace UE::Subsonic
 		UE_API bool GetIsPublic() const;
 
 #if WITH_EDITOR
-		UE_API UE_INTERNAL TArray<FSubsonicEventActionDefinition>& GetMutableActionCollection();
-		UE_API UE_INTERNAL void SetIsPublic(bool bInIsPublic);
+		UE_INTERNAL UE_API TArray<FSubsonicEventActionDefinition>& GetMutableActionCollection();
+		UE_INTERNAL UE_API void SetIsPublic(bool bInIsPublic);
 #endif // WITH_EDITOR
 
 	private:
@@ -162,16 +162,16 @@ namespace UE::Subsonic
 
 #if WITH_EDITOR
 		// Adds action to event with handle if it exists. If insert index set to less than 0 (default), adds to end of collection array.
-		UE_API UE_INTERNAL FSubsonicEventActionDefinition* AddAction(const FEventHandle& InHandle, int32 InsertIndex = INDEX_NONE);
+		UE_INTERNAL UE_API FSubsonicEventActionDefinition* AddAction(const FEventHandle& InHandle, int32 InsertIndex = INDEX_NONE);
 
 		// Adds event with given tag if does not yet exists. Returns pointer to new event if added, null if none added.
-		UE_API UE_INTERNAL FSubsonicEvent* AddEvent(const FGameplayTag& InTag = { }, FSubsonicEvent NewEvent = {});
+		UE_INTERNAL UE_API FSubsonicEvent* AddEvent(const FGameplayTag& InTag = { }, FSubsonicEvent NewEvent = {});
 
 		// Clears all actions under the given event
-		UE_API UE_INTERNAL bool ClearActions(const FEventHandle& InHandle);
+		UE_INTERNAL UE_API bool ClearActions(const FEventHandle& InHandle);
 
 		// Removes all events in the given collection.
-		UE_API UE_INTERNAL void ClearEvents();
+		UE_INTERNAL UE_API void ClearEvents();
 #endif // WITH_EDITOR
 
 		UE_API bool ContainsEvent(FName Name) const;
@@ -188,10 +188,10 @@ namespace UE::Subsonic
 
 #if WITH_EDITOR
 		// Returns pointer to action at handle if it exists
-		UE_API UE_INTERNAL FSubsonicEventActionDefinition* FindMutableAction(const FActionHandle& InHandle);
+		UE_INTERNAL UE_API FSubsonicEventActionDefinition* FindMutableAction(const FActionHandle& InHandle);
 
 		// Returns mutable pointer to event at handle if it exists
-		UE_API UE_INTERNAL FSubsonicEvent* FindMutableEvent(const FEventHandle& InHandle);
+		UE_INTERNAL UE_API FSubsonicEvent* FindMutableEvent(const FEventHandle& InHandle);
 #endif // WITH_EDITOR
 
 		UE_API uint32 GetCollectionId() const;
@@ -210,10 +210,10 @@ namespace UE::Subsonic
 
 #if WITH_EDITOR
 		// Removes action at given handle and returns if successful.
-		UE_API UE_INTERNAL bool RemoveAction(const FActionHandle& InHandle);
+		UE_INTERNAL UE_API bool RemoveAction(const FActionHandle& InHandle);
 
 		// Removes event at given handle and returns if successful.
-		UE_API UE_INTERNAL bool RemoveEvent(const FEventHandle& InHandle);
+		UE_INTERNAL UE_API bool RemoveEvent(const FEventHandle& InHandle);
 #endif // WITH_EDITOR
 
 	private:
@@ -269,7 +269,7 @@ namespace UE::Subsonic
 		UE_API FCollectionHandle GetHandle() const;
 
 #if WITH_EDITOR
-		UE_API UE_INTERNAL FSubsonicEventCollectionDefinition& GetMutableDefinition();
+		UE_INTERNAL UE_API FSubsonicEventCollectionDefinition& GetMutableDefinition();
 #endif // WITH_EDITOR
 
 		UE_API void SetDefinition(FSubsonicEventCollectionDefinition&& Definition);
diff --git a/Engine/Source/Developer/IoStoreUtilities/Internal/ZenStoreHttpClient.h b/Engine/Source/Developer/IoStoreUtilities/Internal/ZenStoreHttpClient.h
--- a/Engine/Source/Developer/IoStoreUtilities/Internal/ZenStoreHttpClient.h
+++ b/Engine/Source/Developer/IoStoreUtilities/Internal/ZenStoreHttpClient.h
@@ -63,10 +63,10 @@ public:
 	UE_API TIoStatusOr<FIoBuffer> ReadChunk(const FIoChunkId& Id, uint64 Offset = 0, uint64 Size = ~0ull);
 	UE_API TIoStatusOr<FIoBuffer> ReadChunk(const FIoHash& RawHash, uint64 Offset = 0, uint64 Size = ~0ull);
 	// TODO: Expand API to add partial chunk requests and ModTag field
-	UE_API UE_INTERNAL void ReadChunks(TConstArrayView<FIoChunkId> Ids, TUniqueFunction<void(const FIoChunkId& Id, TIoStatusOr<FIoBuffer> Result)>&& ChunkCallback);
-	UE_API UE_INTERNAL void ReadChunks(TConstArrayView<FIoHash> RawHashes, TUniqueFunction<void(const FIoHash& RawHash, TIoStatusOr<FIoBuffer> Result)>&& ChunkCallback);
-	UE_API UE_INTERNAL void ReadChunksAsync(TConstArrayView<FIoChunkId> Ids, TUniqueFunction<void(const FIoChunkId& Id, TIoStatusOr<FIoBuffer> Result)>&& ChunkCallback);
-	UE_API UE_INTERNAL void ReadChunksAsync(TConstArrayView<FIoHash> RawHashes, TUniqueFunction<void(const FIoHash& RawHash, TIoStatusOr<FIoBuffer> Result)>&& ChunkCallback);
+	UE_INTERNAL UE_API void ReadChunks(TConstArrayView<FIoChunkId> Ids, TUniqueFunction<void(const FIoChunkId& Id, TIoStatusOr<FIoBuffer> Result)>&& ChunkCallback);
+	UE_INTERNAL UE_API void ReadChunks(TConstArrayView<FIoHash> RawHashes, TUniqueFunction<void(const FIoHash& RawHash, TIoStatusOr<FIoBuffer> Result)>&& ChunkCallback);
+	UE_INTERNAL UE_API void ReadChunksAsync(TConstArrayView<FIoChunkId> Ids, TUniqueFunction<void(const FIoChunkId& Id, TIoStatusOr<FIoBuffer> Result)>&& ChunkCallback);
+	UE_INTERNAL UE_API void ReadChunksAsync(TConstArrayView<FIoHash> RawHashes, TUniqueFunction<void(const FIoHash& RawHash, TIoStatusOr<FIoBuffer> Result)>&& ChunkCallback);
 
 	UE_DEPRECATED(5.6, "Use ReadChunk(FIoHash) instead")
 	UE_API TIoStatusOr<FIoBuffer> ReadOpLogAttachment(FStringView Id);
