Subsonic Parameter Binding Checkpoint
- Finished implementation of action properties that can now be bound to named Collection- or Event-level PropertyBag parameters
- Value propagation: Bound properties auto-update when parameter values change, with type-safe checks (including object subclass compatibility)
Binding menu UX creation/improvements
- Provides deduplicated flat list of parameter names available to be bound to
- Radio buttons pin-types that are colored to type
- "Defined in" tooltips with hyperlink navigation to source parameter bag (event or collection)
- Stale binding handling
- 3-option dialog on rename (rebind/remove/revert)
- 2-option on removal/type change (remove/revert)
- Added SubsonicBindingUtils
free functions
- Removed old SubsonicEventInstanceImpl test code

#rnx

diff --git a/Engine/Plugins/Experimental/Subsonic/Source/Subsonic/Private/CoreActions/SubsonicAction_AudioComponent.h b/Engine/Plugins/Experimental/Subsonic/Source/Subsonic/Private/CoreActions/SubsonicAction_AudioComponent.h
--- a/Engine/Plugins/Experimental/Subsonic/Source/Subsonic/Private/CoreActions/SubsonicAction_AudioComponent.h
+++ b/Engine/Plugins/Experimental/Subsonic/Source/Subsonic/Private/CoreActions/SubsonicAction_AudioComponent.h
@@ -86,7 +86,7 @@ namespace UE::Subsonic
 		FName Name = "Component01";
 
 		// Sound to play.
-		UPROPERTY(EditAnywhere, Category = "Subsonic", meta = (IsBindable))
+		UPROPERTY(EditAnywhere, Category = "Subsonic")
 		TObjectPtr<USoundBase> Sound;
 
 		// Scope of component accessor.  If set to event, provided name is used
diff --git a/Engine/Plugins/Experimental/Subsonic/Source/Subsonic/Private/SubsonicBindingUtils.cpp b/Engine/Plugins/Experimental/Subsonic/Source/Subsonic/Private/SubsonicBindingUtils.cpp
new file mode 100644
--- /dev/null
+++ b/Engine/Plugins/Experimental/Subsonic/Source/Subsonic/Private/SubsonicBindingUtils.cpp
@@ -0,0 +1,276 @@
+// Copyright Epic Games, Inc. All Rights Reserved.
+
+#if WITH_EDITOR
+#include "SubsonicBindingUtils.h"
+
+#include "Algo/AllOf.h"
+#include "Editor.h"
+#include "Misc/MessageDialog.h"
+#include "StructUtils/PropertyBag.h"
+#include "SubsonicEventCollection.h"
+#include "UObject/UnrealType.h"
+
+#define LOCTEXT_NAMESPACE "SubsonicBindingUtils"
+
+
+namespace UE::Subsonic::BindingUtils
+{
+	bool ArePropertiesBindingCompatible(const FProperty* FromBag, const FProperty* ToAction)
+	{
+		check(FromBag);
+		check(ToAction);
+
+		if (ToAction->SameType(FromBag))
+		{
+			return true;
+		}
+		const FObjectPropertyBase* ActionObj = CastField<FObjectPropertyBase>(ToAction);
+		const FObjectPropertyBase* BagObj = CastField<FObjectPropertyBase>(FromBag);
+
+		return ActionObj && BagObj && BagObj->PropertyClass->IsChildOf(ActionObj->PropertyClass);
+	}
+
+#if WITH_EDITORONLY_DATA
+	FStaleBindingsMap FindStaleBindings(
+		const FSubsonicEventCollectionDefinition& Definition, const FCollectionHandle& ParentHandle)
+	{
+		FStaleBindingsMap StaleBindings;
+
+		for (const TPair<FGameplayTag, FSubsonicEvent>& EventPair : Definition.GetEvents())
+		{
+			const FSubsonicEvent& Event = EventPair.Value;
+			const FEventHandle EventHandle = { .Collection = ParentHandle, .EventName = EventPair.Key.GetTagName() };
+
+			int32 ActionIndex = 0;
+			for (const FSubsonicEventActionDefinition& ActionDef : Event.GetActionCollection())
+			{
+				const FActionHandle ActionHandle = { .Event = EventHandle, .Index = ActionIndex++ };
+				const TInstancedStruct<FSubsonicEventActionBase>& Action = ActionDef.GetAction();
+				const FSubsonicEventActionBase* ActionBase = Action.GetPtr();
+				const UScriptStruct* ActionStructType = Action.GetScriptStruct();
+				if (!ActionBase || !ActionStructType)
+				{
+					continue;
+				}
+
+				for (const TPair<FName, FSubsonicEventActionBoundProperties>& BindingPair : ActionBase->GetBindings())
+				{
+					const FName& ParamName = BindingPair.Key;
+					const FPropertyBagPropertyDesc* BagDesc = nullptr;
+					if (const UPropertyBag* EventBag = Event.GetParameters().GetPropertyBagStruct())
+					{
+						BagDesc = EventBag->FindPropertyDescByName(ParamName);
+					}
+					if (!BagDesc)
+					{
+						if (const UPropertyBag* CollBag = Definition.GetParameters().GetPropertyBagStruct())
+						{
+							BagDesc = CollBag->FindPropertyDescByName(ParamName);
+						}
+					}
+
+					const bool bNoParamFound = !BagDesc || !BagDesc->CachedProperty;
+					for (const FName& BoundPropName : BindingPair.Value.Properties)
+					{
+						if (!bNoParamFound)
+						{
+							const FProperty* ActionProp = ActionStructType->FindPropertyByName(BoundPropName);
+							if (ActionProp && BagDesc->CachedProperty && ArePropertiesBindingCompatible(BagDesc->CachedProperty, ActionProp))
+							{
+								continue;
+							}
+						}
+
+						StaleBindings.FindOrAdd(ParamName).Add(FStalePropertyHandle { .Action = ActionHandle, .Property = BoundPropName });
+					}
+				}
+			}
+		}
+
+		return StaleBindings;
+	}
+
+	FName FindRenameCandidate(const FSubsonicEventCollectionDefinition& Definition, FName OldParamName, const TArray<FStalePropertyHandle>& StaleEntries)
+	{
+		// Collect the set of action property types that were bound to the old parameter name.
+		TArray<const FProperty*> StaleActionProps;
+		for (const FStalePropertyHandle& Entry : StaleEntries)
+		{
+			if (const FSubsonicEventActionDefinition* ActionDef = Definition.FindAction(Entry.Action))
+			{
+				if (const UScriptStruct* ActionStruct = ActionDef->GetAction().GetScriptStruct())
+				{
+					if (const FProperty* ActionProp = ActionStruct->FindPropertyByName(Entry.Property))
+					{
+						StaleActionProps.AddUnique(ActionProp);
+					}
+				}
+			}
+		}
+
+		if (StaleActionProps.IsEmpty())
+		{
+			return NAME_None;
+		}
+
+		// Collect all parameter names referenced by ANY binding across the ENTIRE collection,
+		// so that pre-existing parameters bound elsewhere are excluded as rename candidates.
+		TSet<FName> ExistingBoundParams;
+		for (const TPair<FGameplayTag, FSubsonicEvent>& EventPair : Definition.GetEvents())
+		{
+			for (const FSubsonicEventActionDefinition& ActionDef : EventPair.Value.GetActionCollection())
+			{
+				if (const FSubsonicEventActionBase* ActionBase = ActionDef.GetAction().GetPtr())
+				{
+					for (const TPair<FName, FSubsonicEventActionBoundProperties>& BindingPair : ActionBase->GetBindings())
+					{
+						ExistingBoundParams.Add(BindingPair.Key);
+					}
+				}
+			}
+		}
+
+		// Scan a single bag for exactly one candidate parameter that:
+		// 1. Is NOT the old name
+		// 2. Is NOT already referenced by an existing binding
+		// 3. Is type-compatible with the stale action properties
+		// Returns NAME_None if zero or multiple candidates (ambiguous).
+		auto FindCandidateInBag = [&](const UPropertyBag* Bag)
+		{
+			const FName* BagCandidate = nullptr;
+				
+			if (Bag)
+			{
+				for (const FPropertyBagPropertyDesc& Desc : Bag->GetPropertyDescs())
+				{
+					if (Desc.Name != OldParamName && !ExistingBoundParams.Contains(Desc.Name) && Desc.CachedProperty)
+					{
+						const bool bAllCompatible = Algo::AllOf(StaleActionProps, [&Desc](const FProperty* ActionProp)
+						{
+							return ArePropertiesBindingCompatible(Desc.CachedProperty, ActionProp);
+						});
+
+						if (bAllCompatible)
+						{
+							if (!BagCandidate || *BagCandidate == Desc.Name)
+							{
+								BagCandidate = &Desc.Name;
+							}
+							else
+							{
+								// Ambiguous within this bag
+								BagCandidate = nullptr;
+								break;
+							}
+						}
+					}
+				}
+			}
+			return BagCandidate;
+		};
+
+		// Check collection bag first (most common rename target), then each event bag independently.
+		// Scanning per-bag prevents unrelated parameters in other bags from creating false ambiguity.
+		if (const FName* Result = FindCandidateInBag(Definition.GetParameters().GetPropertyBagStruct()))
+		{
+			return *Result;
+		}
+
+		for (const TPair<FGameplayTag, FSubsonicEvent>& EventPair : Definition.GetEvents())
+		{
+			if (const FName* Result = FindCandidateInBag(EventPair.Value.GetParameters().GetPropertyBagStruct()))
+			{
+				return *Result;
+			}
+		}
+
+		return NAME_None;
+	}
+#endif // WITH_EDITORONLY_DATA
+
+	void PromptForStaleBindings(const USubsonicEventCollection& Collection, const FStaleBindingsMap& StaleBindings, EStaleBindingResponse& OutResponse)
+	{
+#if WITH_EDITORONLY_DATA
+		// Detect rename candidates for each stale parameter.
+		TMap<FName, FName> RenameMap;
+		for (const TPair<FName, TArray<FStalePropertyHandle>>& Pair : StaleBindings)
+		{
+			const FName Candidate = FindRenameCandidate(Collection.GetDefinition(), Pair.Key, Pair.Value);
+			if (!Candidate.IsNone())
+			{
+				RenameMap.Add(Pair.Key, Candidate);
+			}
+		}
+
+		const bool bCanRebind = RenameMap.Num() == StaleBindings.Num();
+
+		if (bCanRebind)
+		{
+			// All stale params have rename candidates — offer the rebind option.
+			TArray<FText> RenameTexts;
+			for (const TPair<FName, FName>& Pair : RenameMap)
+			{
+				// \u2192 is the arrow character
+				RenameTexts.Add(FText::Format(LOCTEXT("RenameMapping", "{0} \u2192 {1}"),
+					FText::FromName(Pair.Key), FText::FromName(Pair.Value)));
+			}
+			const FText RenameList = FText::Join(FText::FromString(TEXT("\r\n\t* ")), RenameTexts);
+
+			const FText Title = LOCTEXT("StaleBindingRename_Title", "Parameter Renamed");
+			const FText Message = FText::Format(
+				LOCTEXT("StaleBindingRenameWarning",
+					"The following parameter(s) have been renamed:\r\n\t{0}\r\n"
+					"- Update affected action property bindings to the new name? (Yes)\r\n"
+					"- Remove the stale bindings? (No)\r\n"
+					"- Cancel the parameter rename? (Cancel)"),
+				RenameList);
+
+			const EAppReturnType::Type Result = FMessageDialog::Open(EAppMsgType::YesNoCancel, Message, Title);
+			switch (Result)
+			{
+				case EAppReturnType::Yes:
+				{
+					OutResponse = EStaleBindingResponse::Rebind;
+				}
+				break;
+
+				case EAppReturnType::No:
+				{
+					OutResponse = EStaleBindingResponse::RemoveBindings;
+				}
+				break;
+
+				case EAppReturnType::Cancel:
+				default:
+				{
+					OutResponse = EStaleBindingResponse::Revert;
+				}
+				break;
+			}
+		}
+		else
+		{
+			// No rename candidates — show the standard removal/revert dialog.
+			TArray<FText> ParamTexts;
+			Algo::Transform(StaleBindings, ParamTexts, [](const TPair<FName, TArray<FStalePropertyHandle>>& Pair) { return FText::FromName(Pair.Key); });
+			const FText ParamList = FText::Join(FText::FromString(TEXT("\r\n\t* ")), ParamTexts);
+
+			const FText Title = LOCTEXT("StaleBindingWarning_Title", "Binding Removal Required");
+			const FText Message = FText::Format(
+				LOCTEXT("StaleBindingWarning",
+					"The following parameter(s) have property bindings that are no longer valid "
+					"due to removal or type/name change:\r\n\t{0}\r\n\r\n"
+					"Would you like to remove the affected bindings (Yes) or revert the parameter change (No)?"),
+				ParamList);
+
+			const EAppReturnType::Type Result = FMessageDialog::Open(EAppMsgType::YesNo, Message, Title);
+			OutResponse = Result == EAppReturnType::No
+				? EStaleBindingResponse::Revert
+				: EStaleBindingResponse::RemoveBindings;
+		}
+#endif // WITH_EDITORONLY_DATA
+	}
+} // namespace UE::Subsonic::BindingUtils
+
+#undef LOCTEXT_NAMESPACE
+#endif // WITH_EDITOR
diff --git a/Engine/Plugins/Experimental/Subsonic/Source/Subsonic/Private/SubsonicEventCollection.cpp b/Engine/Plugins/Experimental/Subsonic/Source/Subsonic/Private/SubsonicEventCollection.cpp
--- a/Engine/Plugins/Experimental/Subsonic/Source/Subsonic/Private/SubsonicEventCollection.cpp
+++ b/Engine/Plugins/Experimental/Subsonic/Source/Subsonic/Private/SubsonicEventCollection.cpp
@@ -5,6 +5,7 @@
 #include "Engine/Engine.h"
 #include "NativeGameplayTags.h"
 #include "StructUtils/StructView.h"
+#include "SubsonicBindingUtils.h"
 #include "SubsonicHandles.h"
 #include "SubsonicLog.h"
 #include "SubsonicSubsystem.h"
@@ -30,6 +31,70 @@ namespace UE::Subsonic
 		constexpr uint32 InvalidId = static_cast<uint32>(INDEX_NONE);
 	} // namespace CollectionPrivate
 
+#if WITH_EDITOR
+	const TMap<FName, FSubsonicEventActionBoundProperties>& FSubsonicEventActionBase::GetBindings() const
+	{
+		return Bindings;
+	}
+
+	FName FSubsonicEventActionBase::GetBindingsParameterName()
+	{
+		return GET_MEMBER_NAME_CHECKED(FSubsonicEventActionBase, Bindings);
+	}
+
+	FText FSubsonicEventActionBase::GetDisplayInfo() const
+	{
+		return FText();
+	}
+
+	FName FSubsonicEventActionBase::GetBoundParameterForProperty(FName PropertyName) const
+	{
+#if WITH_EDITORONLY_DATA
+		for (const TPair<FName, FSubsonicEventActionBoundProperties>& Pair : Bindings)
+		{
+			if (Pair.Value.Properties.Contains(PropertyName))
+			{
+				return Pair.Key;
+			}
+		}
+#endif // WITH_EDITORONLY_DATA
+		return NAME_None;
+	}
+
+	void FSubsonicEventActionBase::AddPropertyBinding(FName ParameterName, FName PropertyName)
+	{
+#if WITH_EDITORONLY_DATA
+		RemovePropertyBinding(PropertyName);
+		Bindings.FindOrAdd(ParameterName).Properties.AddUnique(PropertyName);
+#endif // WITH_EDITORONLY_DATA
+	}
+
+	bool FSubsonicEventActionBase::RemoveAllPropertyBindingsToParameter(FName ParameterName)
+	{
+		return Bindings.Remove(ParameterName) > 0;
+	}
+
+	void FSubsonicEventActionBase::RemovePropertyBinding(FName PropertyName)
+	{
+#if WITH_EDITORONLY_DATA
+		TArray<FName> ToRemove;
+		for (TPair<FName, FSubsonicEventActionBoundProperties>& Pair : Bindings)
+		{
+			Pair.Value.Properties.Remove(PropertyName);
+			if (Pair.Value.Properties.IsEmpty())
+			{
+				ToRemove.Add(Pair.Key);
+			}
+		}
+
+		for (FName Name : ToRemove)
+		{
+			Bindings.Remove(Name);
+		}
+#endif // WITH_EDITORONLY_DATA
+	}
+#endif // WITH_EDITOR
+
 	const TInstancedStruct<FSubsonicEventActionBase>& FSubsonicEventActionDefinition::GetAction() const
 	{
 		return Action;
@@ -41,13 +106,104 @@ namespace UE::Subsonic
 		Action = MoveTemp(NewAction);
 	}
 
+	FSubsonicEventActionBase* FSubsonicEventActionDefinition::GetMutableActionBase()
+	{
+		return Action.GetMutablePtr();
+	}
+#endif // WITH_EDITOR
+
+#if WITH_EDITORONLY_DATA
+	const FInstancedPropertyBag& FSubsonicEvent::GetParameters() const
+	{
+		return Parameters;
+	}
+#endif // WITH_EDITORONLY_DATA
+
+#if WITH_EDITOR
 	void FSubsonicEvent::BindParameter(const FInstancedPropertyBag& InSourceParameters, FName ParamName)
+	{
+		const FInstancedPropertyBag* EffectiveBag = nullptr;
+		auto GetPropertyBagDesc = [&ParamName](const FInstancedPropertyBag& Bag)
+		{
+			const FPropertyBagPropertyDesc* Desc = nullptr;
+			if (const UPropertyBag* BagStruct = Bag.GetPropertyBagStruct())
+			{
+				Desc = BagStruct->FindPropertyDescByName(ParamName);
+			}
+
+			return Desc;
+		};
+
+		const FPropertyBagPropertyDesc* PropertyDesc = GetPropertyBagDesc(Parameters);
+		if (PropertyDesc)
+		{
+			EffectiveBag = &Parameters;
+		}
+		else
+		{
+			PropertyDesc = GetPropertyBagDesc(InSourceParameters);
+			EffectiveBag = &InSourceParameters;
+		}
+
+		if (!PropertyDesc || !PropertyDesc->CachedProperty)
+		{
+			return;
+		}
+
+		check(EffectiveBag);
+		const FConstStructView BagView = EffectiveBag->GetValue();
+
+		for (FSubsonicEventActionDefinition& ActionDefinition : ActionCollection)
+		{
+			const UScriptStruct* ActionStructType = ActionDefinition.Action.GetScriptStruct();
+			FSubsonicEventActionBase* ActionPtr = ActionDefinition.Action.GetMutablePtr();
+			if (!ActionPtr || !ActionStructType)
+			{
+				continue;
+			}
+
+			// Propagate parameter value to all type-compatible action properties.
+			// This does not record entries in the Bindings map — only explicit bindings
+			// created via AddPropertyBinding (the editor UI) are tracked there.
+			// Skip properties that are explicitly bound to a different parameter so that
+			// auto-propagation of one parameter cannot overwrite an explicit binding's value.
+			for (TFieldIterator<FProperty> PropIt(ActionStructType); PropIt; ++PropIt)
+			{
+				FProperty* ActionProp = *PropIt;
+				if (ActionProp->HasMetaData("NoBinding"))
+				{
+					continue;
+				}
+
+				// Require compatible property types (including inner types for containers) to safely copy the value.
+				if (!BindingUtils::ArePropertiesBindingCompatible(PropertyDesc->CachedProperty, ActionProp))
+				{
+					continue;
+				}
+
+				// Only propagate to properties explicitly bound to this parameter.
+				const FName BoundParam = ActionPtr->GetBoundParameterForProperty(ActionProp->GetFName());
+				if (BoundParam != ParamName)
+				{
+					continue;
+				}
+
+				if (const void* SrcPtr = PropertyDesc->CachedProperty->ContainerPtrToValuePtr<const void>(BagView.GetMemory()); ensure(SrcPtr))
+				{
+					void* DstPtr = ActionProp->ContainerPtrToValuePtr<void>(ActionPtr);
+					ActionProp->CopyCompleteValue(DstPtr, SrcPtr);
+				}
+			}
+		}
+	}
+
+	void FSubsonicEvent::UnbindParameter(FName ParamName)
 	{
 		for (FSubsonicEventActionDefinition& ActionDefinition : ActionCollection)
 		{
 			if (FSubsonicEventActionBase* ActionPtr = ActionDefinition.Action.GetMutablePtr())
 			{
-				ActionPtr->BindParameter(InSourceParameters, ParamName);
+				ActionPtr->RemoveAllPropertyBindingsToParameter(ParamName);
 			}
 		}
 	}
@@ -55,7 +211,8 @@ namespace UE::Subsonic
 
 	void FSubsonicEvent::Execute(const FSubsonicExecutor& InExecutor, const FEventHandle& InHandle) const
 	{
-		checkf(InExecutor.GetCollectionHandle().CollectionId == InHandle.Collection.CollectionId, TEXT("Invalid execution context: executor CollectionId must match that of the given event handle"));
+		checkf(InExecutor.GetCollectionHandle().CollectionId == InHandle.Collection.CollectionId,
+			TEXT("Invalid execution context: executor CollectionId must match that of the given event handle"));
 
 		FActionHandle ActionHandle { .Event = InHandle };
 		for (int32 Index = 0; Index < ActionCollection.Num(); ++Index)
@@ -84,6 +241,11 @@ namespace UE::Subsonic
 	{
 		return GET_MEMBER_NAME_CHECKED(FSubsonicEvent, ActionCollection);
 	}
+
+	FName FSubsonicEvent::GetParametersPropertyName()
+	{
+		return GET_MEMBER_NAME_CHECKED(FSubsonicEvent, Parameters);
+	}
 #endif // WITH_EDITOR
 
 #if WITH_EDITOR
@@ -177,6 +339,35 @@ namespace UE::Subsonic
 		return &Events.Add(InTag, MoveTemp(NewEvent));
 	}
 
+	void FSubsonicEventCollectionDefinition::BindAllParameters()
+	{
+		if (const UPropertyBag* BagStruct = Parameters.GetPropertyBagStruct())
+		{
+			for (const FPropertyBagPropertyDesc& Desc : BagStruct->GetPropertyDescs())
+			{
+				BindParameter(Desc.Name);
+			}
+		}
+
+		// Bind event-level parameters that have no collection-level counterpart
+		for (TPair<FGameplayTag, FSubsonicEvent>& Pair : Events)
+		{
+			const UPropertyBag* EventBagStruct = Pair.Value.Parameters.GetPropertyBagStruct();
+			if (!EventBagStruct)
+			{
+				continue;
+			}
+			for (const FPropertyBagPropertyDesc& Desc : EventBagStruct->GetPropertyDescs())
+			{
+				const UPropertyBag* CollectionBagStruct = Parameters.GetPropertyBagStruct();
+				if (!CollectionBagStruct || !CollectionBagStruct->FindPropertyDescByName(Desc.Name))
+				{
+					Pair.Value.BindParameter(Parameters, Desc.Name);
+				}
+			}
+		}
+	}
+
 	void FSubsonicEventCollectionDefinition::BindParameter(FName ParamName)
 	{
 		for (TPair<FGameplayTag, FSubsonicEvent>& Pair : Events)
@@ -184,6 +375,107 @@ namespace UE::Subsonic
 			Pair.Value.BindParameter(Parameters, ParamName);
 		}
 	}
+
+	bool FSubsonicEventCollectionDefinition::RemoveAllPropertyBindingsToParameter(FName ParamName)
+	{
+		bool bBindingsRemoved = false;
+
+#if WITH_EDITORONLY_DATA
+		for (TPair<FGameplayTag, FSubsonicEvent>& EventPair : Events)
+		{
+			for (FSubsonicEventActionDefinition& ActionDef : EventPair.Value.GetMutableActionCollection())
+			{
+				if (FSubsonicEventActionBase* ActionBase = ActionDef.GetMutableActionBase())
+				{
+					bBindingsRemoved |= ActionBase->RemoveAllPropertyBindingsToParameter(ParamName);
+				}
+			}
+		}
+#endif // WITH_EDITORONLY_DATA
+
+		return bBindingsRemoved;
+	}
+
+	bool FSubsonicEventCollectionDefinition::RemoveStaleBindings(const FCollectionHandle& ParentHandle, bool bPromptRemoval)
+	{
+#if WITH_EDITORONLY_DATA
+		using EStaleBindingResponse = BindingUtils::EStaleBindingResponse;
+		using FStalePropertyHandle = BindingUtils::FStalePropertyHandle;
+
+		// Check for stale bindings caused by structural changes (type or name changes) to parameter bags.
+		const TMap<FName, TArray<FStalePropertyHandle>> StaleBindings = BindingUtils::FindStaleBindings(*this, ParentHandle);
+		if (!StaleBindings.IsEmpty())
+		{
+			if (bPromptRemoval)
+			{ 
+				EStaleBindingResponse Response = EStaleBindingResponse::RemoveBindings;
+				OnStaleBindingsDetected.Broadcast(StaleBindings, Response);
+				if (Response == EStaleBindingResponse::Revert)
+				{
+					// If an editor is open, it is expected to respond to the delegate above and undo the associated transaction in the
+					// next tick. A response of revert essentially signals the remaining removal to be skipped. (If no editor is open,
+					// stale bindings are removed and stale bindings are removed on next editor load).
+					return false;
+				}
+
+				if (Response == EStaleBindingResponse::Rebind)
+				{
+					// Update binding keys from old parameter name to the detected rename target.
+					// Falls back to removal for any parameter where no unique candidate is found.
+					for (const TPair<FName, TArray<FStalePropertyHandle>>& Pair : StaleBindings)
+					{
+						const FName NewParamName = BindingUtils::FindRenameCandidate(*this, Pair.Key, Pair.Value);
+						for (const FStalePropertyHandle& Entry : Pair.Value)
+						{
+							if (FSubsonicEventActionDefinition* Action = FindMutableAction(Entry.Action))
+							{
+								if (FSubsonicEventActionBase* ActionBase = Action->GetMutableActionBase())
+								{
+									ActionBase->RemovePropertyBinding(Entry.Property);
+									if (!NewParamName.IsNone())
+									{
+										ActionBase->AddPropertyBinding(NewParamName, Entry.Property);
+									}
+								}
+							}
+						}
+					}
+					return true;
+				}
+			}
+		}
+		else
+		{
+			return false;
+		}
+
+		for (const TPair<FName, TArray<FStalePropertyHandle>>& Pair : StaleBindings)
+		{
+			for (const FStalePropertyHandle& Entry : Pair.Value)
+			{
+				if (FSubsonicEventActionDefinition* Action = FindMutableAction(Entry.Action))
+				{
+					if (FSubsonicEventActionBase* ActionBase = Action->GetMutableActionBase())
+					{
+						ActionBase->RemovePropertyBinding(Entry.Property);
+					}
+				}
+			}
+		}
+
+		return true;
+#else // !WITH_EDITORONLY_DATA
+		return false;
+#endif // !WITH_EDITORONLY_DATA
+	}
+
+	void FSubsonicEventCollectionDefinition::UnbindParameter(FName ParamName)
+	{
+		for (TPair<FGameplayTag, FSubsonicEvent>& Pair : Events)
+		{
+			Pair.Value.UnbindParameter(ParamName);
+		}
+	}
 #endif // WITH_EDITOR
 
 	FSubsonicEventCollectionDefinition FSubsonicEventCollectionDefinition::Create(FName Name, TMap<FGameplayTag, FSubsonicEvent> InEvents)
@@ -226,6 +518,29 @@ namespace UE::Subsonic
 	{
 		Events.Empty();
 	}
+
+	bool FSubsonicEventCollectionDefinition::ContainsBinding(FName EventName, FName ParameterName) const
+	{
+		if (Parameters.FindPropertyDescByName(ParameterName) != nullptr)
+		{
+			return true;
+		}
+
+		if (const FSubsonicEvent* Event = FindEvent(EventName))
+		{
+			if (Event->GetParameters().FindPropertyDescByName(ParameterName) != nullptr)
+			{
+				return true;
+			}
+		}
+
+		return false;
+	}
+
+	BindingUtils::FOnStaleBindingsDetected& FSubsonicEventCollectionDefinition::GetOnStaleBindingsDetectedDelegate()
+	{
+		return OnStaleBindingsDetected;
+	}
 #endif // WITH_EDITOR
 
 	bool FSubsonicEventCollectionDefinition::ContainsEvent(const FGameplayTag& EventTag) const
@@ -355,6 +670,13 @@ namespace UE::Subsonic
 	}
 #endif // WITH_EDITOR
 
+#if WITH_EDITORONLY_DATA
+	const FInstancedPropertyBag& FSubsonicEventCollectionDefinition::GetParameters() const
+	{
+		return Parameters;
+	}
+#endif // WITH_EDITORONLY_DATA
+
 	bool FSubsonicEventCollectionDefinition::IsValid() const
 	{
 		return CollectionId != INDEX_NONE;
@@ -452,6 +774,37 @@ namespace UE::Subsonic
 	{
 		return CollectionDefinition;
 	}
+
+	void USubsonicEventCollection::PostEditChangeChainProperty(FPropertyChangedChainEvent& PropertyChangedEvent)
+	{
+		Super::PostEditChangeChainProperty(PropertyChangedEvent);
+
+#if WITH_EDITORONLY_DATA
+		for (auto* Node = PropertyChangedEvent.PropertyChain.GetHead(); Node; Node = Node->GetNextNode())
+		{
+			const FName PropName = Node->GetValue()->GetFName();
+			if (PropName == GET_MEMBER_NAME_CHECKED(FSubsonicEventCollectionDefinition, Parameters)
+				|| PropName == FSubsonicEvent::GetParametersPropertyName()
+				|| PropName == FSubsonicEventActionBase::GetBindingsParameterName())
+			{
+				const bool bBindingsRemoved = CollectionDefinition.RemoveStaleBindings(GetHandle(), true /* bPromptRemoval */);
+				if (bBindingsRemoved)
+				{
+					CollectionDefinition.BindAllParameters();
+				}
+				break;
+			}
+		}
+#endif // WITH_EDITORONLY_DATA
+	}
+
+	void USubsonicEventCollection::PostEditUndo()
+	{
+		Super::PostEditUndo();
+
+		CollectionDefinition.RemoveStaleBindings(GetHandle());
+		CollectionDefinition.BindAllParameters();
+	}
 #endif // WITH_EDITOR
 
 
diff --git a/Engine/Plugins/Experimental/Subsonic/Source/Subsonic/Private/SubsonicEventInstanceImpl.cpp b/Engine/Plugins/Experimental/Subsonic/Source/Subsonic/Private/SubsonicEventInstanceImpl.cpp
deleted file mode 100644
diff --git a/Engine/Plugins/Experimental/Subsonic/Source/Subsonic/Private/SubsonicEventInstanceImpl.h b/Engine/Plugins/Experimental/Subsonic/Source/Subsonic/Private/SubsonicEventInstanceImpl.h
deleted file mode 100644
diff --git a/Engine/Plugins/Experimental/Subsonic/Source/Subsonic/Private/SubsonicSubsystem.cpp b/Engine/Plugins/Experimental/Subsonic/Source/Subsonic/Private/SubsonicSubsystem.cpp
--- a/Engine/Plugins/Experimental/Subsonic/Source/Subsonic/Private/SubsonicSubsystem.cpp
+++ b/Engine/Plugins/Experimental/Subsonic/Source/Subsonic/Private/SubsonicSubsystem.cpp
@@ -10,7 +10,6 @@
 #include "Misc/CoreDefines.h"
 #include "Misc/Optional.h"
 #include "SubsonicEventCollection.h"
-#include "SubsonicEventInstanceImpl.h"
 #include "SubsonicHandles.h"
 #include "UObject/GCObject.h"
 
diff --git a/Engine/Plugins/Experimental/Subsonic/Source/Subsonic/Public/SubsonicBindingUtils.h b/Engine/Plugins/Experimental/Subsonic/Source/Subsonic/Public/SubsonicBindingUtils.h
new file mode 100644
--- /dev/null
+++ b/Engine/Plugins/Experimental/Subsonic/Source/Subsonic/Public/SubsonicBindingUtils.h
@@ -0,0 +1,67 @@
+// Copyright Epic Games, Inc. All Rights Reserved.
+#pragma once
+
+#if WITH_EDITOR
+#include "Containers/Set.h"
+#include "Delegates/DelegateCombinations.h"
+#include "SubsonicHandles.h"
+
+
+#define UE_API SUBSONIC_API
+
+namespace UE::Subsonic
+{
+	// Forward Declarations
+	class USubsonicEventCollection;
+
+	namespace BindingUtils
+	{
+		// Response returned by OnStaleBindingsDetected to indicate whether to remove all stale
+		// bindings (proceeding with the parameter structural change), revert the change, or rebind
+		// stale bindings to a renamed parameter (where applicable).
+		enum class EStaleBindingResponse : uint8
+		{
+			RemoveBindings,
+			Revert,
+			Rebind
+		};
+
+		// Represents minimal information required to locate a stale binding entry.
+		struct FStalePropertyHandle
+		{
+			// Action containing stale binding.
+			FActionHandle Action;
+
+			// Name of property on action that is stale
+			FName Property;
+		};
+
+		// Map of stale parameter name to action pointing to stale property
+		using FStaleBindingsMap = TMap<FName, TArray<FStalePropertyHandle>>;
+
+		// Fired when a structural parameter change (type or name change) creates stale editor bindings.
+		DECLARE_MULTICAST_DELEGATE_TwoParams(FOnStaleBindingsDetected, const FStaleBindingsMap&, EStaleBindingResponse&);
+
+		// Checks whether a property bag value can be safely copied into a destination (action) property.
+		// Uses SameType for exact matching, with a fallback for object properties to allow subclass
+		// assignment. Does not currently handle inner-type variance for container properties (TArray, etc.).
+		UE_INTERNAL UE_API bool ArePropertiesBindingCompatible(const FProperty* FromBag, const FProperty* ToAction);
+
+		// For a stale parameter name, finds a likely rename candidate in the current bags: a parameter
+		// with a compatible type that is not already referenced by any existing binding on the affected
+		// actions. Returns NAME_None if no unique candidate is found.
+		UE_INTERNAL UE_API FName FindRenameCandidate(const FSubsonicEventCollectionDefinition& Definition, FName OldParamName, const TArray<FStalePropertyHandle>& StaleEntries);
+
+		// Returns map of parameter names to stale binding entries for bindings that are no longer valid
+		// (parameter removed or type no longer matches). Checks event-level bags first, then collection-level.
+		UE_INTERNAL UE_API TMap<FName, TArray<FStalePropertyHandle>> FindStaleBindings(const FSubsonicEventCollectionDefinition& Definition, const FCollectionHandle& ParentHandle);
+
+		// Prompts the user to handle stale bindings caused by a parameter rename or removal.
+		// Detects rename candidates and shows either a 3-option (rename/remove/revert) or 2-option
+		// (remove/revert) dialog. Sets OutResponse accordingly.
+		UE_INTERNAL UE_API void PromptForStaleBindings(const USubsonicEventCollection& Collection, const FStaleBindingsMap& StaleBindings, EStaleBindingResponse& OutResponse);
+	} // namespace BindingUtils
+} // namespace UE::Subsonic
+
+#undef UE_API // SUBSONIC_API
+#endif // WITH_EDITOR
diff --git a/Engine/Plugins/Experimental/Subsonic/Source/Subsonic/Public/SubsonicEventCollection.h b/Engine/Plugins/Experimental/Subsonic/Source/Subsonic/Public/SubsonicEventCollection.h
--- a/Engine/Plugins/Experimental/Subsonic/Source/Subsonic/Public/SubsonicEventCollection.h
+++ b/Engine/Plugins/Experimental/Subsonic/Source/Subsonic/Public/SubsonicEventCollection.h
@@ -4,10 +4,9 @@
 #include "AudioDefines.h"
 #include "GameplayTagContainer.h"
 #include "NativeGameplayTags.h"
-#include "PropertyBindingBindingCollectionOwner.h"
-#include "PropertyBindingPath.h"
 #include "StructUtils/InstancedStruct.h"
 #include "StructUtils/PropertyBag.h"
+#include "SubsonicBindingUtils.h"
 #include "SubsonicExecutor.h"
 #include "UObject/NameTypes.h"
 
@@ -22,6 +21,12 @@ namespace UE::Subsonic
 {
 	// Forward Declarations
 	struct FActionHandle;
+	class USubsonicEventCollection;
+
+	namespace BindingUtils
+	{
+		enum class EStaleBindingResponse : uint8;
+	} // namespace BindingUtils
 
 	UE_API UE_DECLARE_GAMEPLAY_TAG_EXTERN(TAG_SubsonicCore)
 	UE_API UE_DECLARE_GAMEPLAY_TAG_EXTERN(TAG_SubsonicCore_Event_Play)
@@ -44,16 +49,31 @@ namespace UE::Subsonic
 		virtual ~FSubsonicEventActionBase() = default;
 
 #if WITH_EDITOR
+		UE_INTERNAL UE_API const TMap<FName, FSubsonicEventActionBoundProperties>& GetBindings() const;
+		UE_INTERNAL UE_API static FName GetBindingsParameterName();
+
 		// Returns display info that is custom to the given action instance to be shown
 		// adjacent to the action, for example, in the Subsonic Event Editor action tree.
-		UE_API virtual FText GetDisplayInfo() const { return FText(); }
+		UE_API virtual FText GetDisplayInfo() const;
 #endif // WITH_EDITOR
 
-	protected:
+	public:
 #if WITH_EDITOR
-		virtual void BindParameter(const FInstancedPropertyBag& InBag, FName Name) { }
+		// Returns the parameter name that the given action property is currently bound to, or NAME_None if unbound.
+		UE_API FName GetBoundParameterForProperty(FName PropertyName) const;
+
+		// Binds an action property to a named parameter, replacing any previous binding for that property.
+		UE_API void AddPropertyBinding(FName ParameterName, FName PropertyName);
+
+		// Removes all property bindings associated with the parameter of the given name.
+		// Returns true if binding was found and removed, false if not.
+		UE_API bool RemoveAllPropertyBindingsToParameter(FName ParameterName);
+
+		// Removes any binding for the given action property.
+		UE_API void RemovePropertyBinding(FName PropertyName);
 #endif // WITH_EDITOR
 
+	protected:
 		virtual void Execute(const FSubsonicExecutor& InExecutor, const FActionHandle& InHandle) const { }
 
 #if WITH_EDITORONLY_DATA
@@ -74,6 +94,9 @@ namespace UE::Subsonic
 
 #if WITH_EDITOR
 		UE_API void SetAction(TInstancedStruct<FSubsonicEventActionBase> NewAction);
+
+		// Returns a mutable pointer to the underlying action base, or nullptr if no action is set.
+		UE_INTERNAL UE_API FSubsonicEventActionBase* GetMutableActionBase();
 #endif // WITH_EDITOR
  
 	private:
@@ -92,10 +115,15 @@ namespace UE::Subsonic
 
 #if WITH_EDITOR
 		UE_INTERNAL UE_API static FName GetActionCollectionPropertyName();
+		UE_INTERNAL UE_API static FName GetParametersPropertyName();
 #endif // WITH_EDITOR
 
 		UE_API const TArray<FSubsonicEventActionDefinition>& GetActionCollection() const;
 
+#if WITH_EDITORONLY_DATA
+		UE_API const FInstancedPropertyBag& GetParameters() const;
+#endif // WITH_EDITORONLY_DATA
+
 #if WITH_EDITOR
 		UE_API bool GetAutoAudition() const;
 #endif // WITH_EDITOR
@@ -113,6 +141,7 @@ namespace UE::Subsonic
 
 #if WITH_EDITOR
 		UE_INTERNAL void BindParameter(const FInstancedPropertyBag& InSourceParameters, FName Name);
+		UE_INTERNAL void UnbindParameter(FName Name);
 #endif // WITH_EDITOR
 
 	private:
@@ -172,6 +201,15 @@ namespace UE::Subsonic
 
 		// Removes all events in the given collection.
 		UE_INTERNAL UE_API void ClearEvents();
+
+		// Returns true if binding exists with the given parameter name on either the collection or event with
+		// the given name, false if not. Also returns false if event with the given name is not found.
+		UE_INTERNAL UE_API bool ContainsBinding(FName EventName, FName ParameterName) const;
+
+		// Fired when a structural parameter change creates stale editor bindings.
+		// The handler should show appropriate UI and return the desired response.
+		// If unbound, stale bindings are removed automatically.
+		UE_INTERNAL UE_API BindingUtils::FOnStaleBindingsDetected& GetOnStaleBindingsDetectedDelegate();
 #endif // WITH_EDITOR
 
 		UE_API bool ContainsEvent(FName Name) const;
@@ -206,6 +244,10 @@ namespace UE::Subsonic
 		UE_API static FName GetParametersPropertyName();
 #endif // WITH_EDITOR
 
+#if WITH_EDITORONLY_DATA
+		UE_API const FInstancedPropertyBag& GetParameters() const;
+#endif // WITH_EDITORONLY_DATA
+
 		UE_API bool IsValid() const;
 
 #if WITH_EDITOR
@@ -214,10 +256,18 @@ namespace UE::Subsonic
 
 		// Removes event at given handle and returns if successful.
 		UE_INTERNAL UE_API bool RemoveEvent(const FEventHandle& InHandle);
+
+		// Removes all stale property bindings. Defined as stale if type or
+		// name of mapped property does not match that of target parameter.
+		// If bPromptRemoval is set, delegate is executed enabling editors to
+		// prompt user to cancel removal. Returns whether or not removal succeeded.
+		UE_INTERNAL UE_API bool RemoveStaleBindings(const FCollectionHandle& ParentHandle, bool bPromptRemoval = false);
 #endif // WITH_EDITOR
 
 	private:
 #if WITH_EDITOR
+		BindingUtils::FOnStaleBindingsDetected OnStaleBindingsDetected;
+
 		template <typename ValueType>
 		bool GetParameterValue(FName ParamName, ValueType& OutValue) const
 		{
@@ -225,6 +275,12 @@ namespace UE::Subsonic
 		}
 
 		UE_INTERNAL void BindParameter(FName Parameter);
+		UE_INTERNAL void BindAllParameters();
+		UE_INTERNAL void UnbindParameter(FName Parameter);
+
+		// Removes all editor bindings to the given parameter name from all actions
+		// in all events. Returns whether or not any bindings were removed.
+		UE_INTERNAL UE_API bool RemoveAllPropertyBindingsToParameter(FName ParamName);
 #endif // WITH_EDITOR
 
 		// Register's the given collection with the Subsonic Subsystem.
@@ -281,6 +337,11 @@ namespace UE::Subsonic
 		virtual void PostInitProperties() override;
 		virtual void PostLoad() override;
 		virtual void BeginDestroy() override;
+
+#if WITH_EDITOR
+		virtual void PostEditChangeChainProperty(FPropertyChangedChainEvent& PropertyChangedEvent) override;
+		virtual void PostEditUndo() override;
+#endif // WITH_EDITOR
 	};
 
 	// Executor handle for a Subsonic Event Collection.
diff --git a/Engine/Plugins/Experimental/Subsonic/Source/Subsonic/Public/SubsonicSubsystem.h b/Engine/Plugins/Experimental/Subsonic/Source/Subsonic/Public/SubsonicSubsystem.h
--- a/Engine/Plugins/Experimental/Subsonic/Source/Subsonic/Public/SubsonicSubsystem.h
+++ b/Engine/Plugins/Experimental/Subsonic/Source/Subsonic/Public/SubsonicSubsystem.h
@@ -17,7 +17,6 @@
 namespace UE::Subsonic
 {
 	// Forward Declarations
-	class FSubsonicEventInstanceImpl;
 	class USubsonicEventCollection;
 
 	namespace SubsystemPrivate
diff --git a/Engine/Plugins/Experimental/Subsonic/Source/SubsonicEditor/Private/SubsonicEditorModule.cpp b/Engine/Plugins/Experimental/Subsonic/Source/SubsonicEditor/Private/SubsonicEditorModule.cpp
--- a/Engine/Plugins/Experimental/Subsonic/Source/SubsonicEditor/Private/SubsonicEditorModule.cpp
+++ b/Engine/Plugins/Experimental/Subsonic/Source/SubsonicEditor/Private/SubsonicEditorModule.cpp
@@ -6,6 +6,8 @@
 #include "SubsonicSlateStyle.h"
 
 
+#define LOCTEXT_NAMESPACE "SubsonicEditor"
+
 namespace UE::Subsonic::Editor
 {
 	class FModule : public IModuleInterface
@@ -68,4 +70,6 @@ namespace UE::Subsonic::Editor
 	};
 } // namespace UE::Subsonic::Editor
 
+#undef LOCTEXT_NAMESPACE
+
 IMPLEMENT_MODULE(UE::Subsonic::Editor::FModule, SubsonicEditor);
diff --git a/Engine/Plugins/Experimental/Subsonic/Source/SubsonicEditor/Private/SubsonicEventCollectionEditor.cpp b/Engine/Plugins/Experimental/Subsonic/Source/SubsonicEditor/Private/SubsonicEventCollectionEditor.cpp
--- a/Engine/Plugins/Experimental/Subsonic/Source/SubsonicEditor/Private/SubsonicEventCollectionEditor.cpp
+++ b/Engine/Plugins/Experimental/Subsonic/Source/SubsonicEditor/Private/SubsonicEventCollectionEditor.cpp
@@ -2,6 +2,7 @@
 
 #include "SubsonicEventCollectionEditor.h"
 
+#include "Algo/AnyOf.h"
 #include "AudioDeviceHandle.h"
 #include "AudioDeviceManager.h"
 #include "Engine/Engine.h"
@@ -10,18 +11,21 @@
 #include "GameplayTagContainer.h"
 #include "Input/DragAndDrop.h"
 #include "ISinglePropertyView.h"
+#include "Misc/MessageDialog.h"
 #include "Modules/ModuleManager.h"
-#include "PropertyBindingExtension.h"
 #include "PropertyEditorModule.h"
 #include "ScopedTransaction.h"
 #include "SGameplayTagCombo.h"
 #include "Styling/SlateStyleRegistry.h"
+#include "SubsonicBindingUtils.h"
 #include "SubsonicEditorSubsystem.h"
 #include "SubsonicEventCollection.h"
 #include "SubsonicEventCollectionEditorCommands.h"
+#include "SubsonicPropertyBindingExtension.h"
 #include "Widgets/Docking/SDockTab.h"
 #include "Widgets/Input/SSearchableComboBox.h"
 #include "Widgets/Input/SSearchBox.h"
+#include "Widgets/Layout/SBox.h"
 #include "Widgets/SNullWidget.h"
 #include "Widgets/Views/SListView.h"
 #include "Widgets/Views/STableRow.h"
@@ -31,12 +35,36 @@
 
 namespace UE::Subsonic::Editor
 {
+	void TransactEventCollection(FText Description, USubsonicEventCollection& Collection, TFunctionRef<void(FSubsonicEventCollectionDefinition&)> TransactionFunc)
+	{
+		FScopedTransaction Transaction(Description);
+		Collection.Modify();
+		TransactionFunc(Collection.GetMutableDefinition());
+	}
+
 	namespace EventCollectionEditorPrivate
 	{
 		const FName TabName_EventTree = "SubsonicEventCollectionEditor_EventTree";
 		const FName TabName_Inspector = "SubsonicEventCollectionEditor_Inspector";
 		const FName TabName_Parameters = "SubsonicEventCollectionEditor_Parameters";
 
+		class FActionDragDropOp : public FDragDropOperation
+		{
+		public:
+			DRAG_DROP_OPERATOR_TYPE(FActionDragDropOp, FDragDropOperation)
+
+				FActionHandle SourceAction;
+
+			static TSharedRef<FActionDragDropOp> New(FActionHandle InAction)
+			{
+				TSharedRef<FActionDragDropOp> Op = MakeShared<FActionDragDropOp>();
+				Op->SourceAction = MoveTemp(InAction);
+				Op->MouseCursor = EMouseCursor::GrabHandClosed;
+				Op->Construct();
+				return Op;
+			}
+		};
+
 		Audio::FDeviceId GetPreviewDeviceId()
 		{
 			if (FAudioDeviceManager* DeviceManager = GEngine->GetAudioDeviceManager())
@@ -71,23 +99,6 @@ namespace UE::Subsonic::Editor
 			return NullBrush;
 		}
 
-		class FActionDragDropOp : public FDragDropOperation
-		{
-		public:
-			DRAG_DROP_OPERATOR_TYPE(FActionDragDropOp, FDragDropOperation)
-
-			FActionHandle SourceAction;
-
-			static TSharedRef<FActionDragDropOp> New(FActionHandle InAction)
-			{
-				TSharedRef<FActionDragDropOp> Op = MakeShared<FActionDragDropOp>();
-				Op->SourceAction = MoveTemp(InAction);
-				Op->MouseCursor = EMouseCursor::GrabHandClosed;
-				Op->Construct();
-				return Op;
-			}
-		};
-
 		TSharedPtr<IPropertyHandle> GetEventTagPropertyHandle(TSharedRef<ISinglePropertyView> CollectionView, FName EventName)
 		{
 			TSharedPtr<IPropertyHandle> CollectionProperty = CollectionView->GetPropertyHandle();
@@ -558,6 +569,30 @@ namespace UE::Subsonic::Editor
 		}
 	}
 
+	void FEventCollectionEditor::SelectEventInTree(FName EventName)
+	{
+		if (!EventTreeView.IsValid())
+		{
+			return;
+		}
+		for (const TSharedPtr<FHandleVariant>& Handle : EventTreeHandles)
+		{
+			if (Handle.IsValid() && Handle->IsEvent() && Handle->GetEvent().EventName == EventName)
+			{
+				EventTreeView->SetSelection(Handle, ESelectInfo::Direct);
+				return;
+			}
+		}
+	}
+
+	void FEventCollectionEditor::TryInvokeTab(FName TabName) const
+	{
+		if (TabManager.IsValid())
+		{
+			TabManager->TryInvokeTab(TabName);
+		}
+	}
+
 	void FEventCollectionEditor::RebuildEventTree(bool bRebuildListView)
 	{
 		ActionComboValues.Reset();
@@ -677,7 +712,24 @@ namespace UE::Subsonic::Editor
 		.SetDisplayName(LOCTEXT("InspectorTab", "Inspector"))
 		.SetGroup(WorkspaceMenuCategory.ToSharedRef())
 		.SetIcon(FSlateIcon(FAppStyle::GetAppStyleSetName(), "LevelEditor.Tabs.Properties"));
-		EventTreeInspectorView->SetExtensionHandler(MakeShared<FPropertyBindingExtension>());
+		TWeakPtr<FEventCollectionEditor> WeakThis(SharedThis(this));
+		EventTreeInspectorView->SetExtensionHandler(MakeShared<FSubsonicPropertyBindingExtension>(
+			[WeakThis](FName EventName)
+			{
+				if (TSharedPtr<FEventCollectionEditor> Editor = WeakThis.Pin())
+				{
+					// None = collection level parameters, so just invoke/bring focus to that tab
+					if (EventName.IsNone())
+					{
+						Editor->GetTabManager()->TryInvokeTab(EventCollectionEditorPrivate::TabName_Parameters);
+					}
+					else
+					{
+						Editor->SelectEventInTree(EventName);
+					}
+				}
+			}
+		));
 
 		FAssetEditorToolkit::RegisterTabSpawners(InTabManager);
 	}
@@ -693,6 +745,11 @@ namespace UE::Subsonic::Editor
 
 	FEventCollectionEditor::~FEventCollectionEditor()
 	{
+		if (USubsonicEventCollection* Collection = GetEventCollection())
+		{
+			Collection->GetMutableDefinition().GetOnStaleBindingsDetectedDelegate().Remove(StaleBindingsDelegateHandle);
+		}
+
 		StopAudition();
 
 		if (GEditor)
@@ -701,6 +758,53 @@ namespace UE::Subsonic::Editor
 		}
 	}
 
+	void FEventCollectionEditor::InitBindings()
+	{
+		using namespace BindingUtils;
+
+		USubsonicEventCollection* Collection = GetEventCollection();
+		if (!Collection)
+		{
+			return;
+		}
+
+		FSubsonicEventCollectionDefinition& Definition = Collection->GetMutableDefinition();
+		Definition.RemoveStaleBindings(Collection->GetHandle());
+		FOnStaleBindingsDetected& StaleBindingsDelegate = Definition.GetOnStaleBindingsDetectedDelegate();
+		StaleBindingsDelegate.Remove(StaleBindingsDelegateHandle);
+
+		// Register handler to show a confirmation dialog when a structural parameter change
+		// (type or name change) would invalidate editor property bindings.
+		StaleBindingsDelegateHandle = StaleBindingsDelegate.AddLambda([WeakCol = TWeakObjectPtr<USubsonicEventCollection>(Collection)](const FStaleBindingsMap& StaleBindings, EStaleBindingResponse& OutResponse)
+		{
+			// Ignore if another delegate handler responded that it reverted on stale binding detection.
+			if (OutResponse != EStaleBindingResponse::Revert)
+			{
+				if (const USubsonicEventCollection* PinnedCol = WeakCol.Get())
+				{
+					BindingUtils::PromptForStaleBindings(*PinnedCol, StaleBindings, OutResponse);
+				}
+
+				// Deferred to next tick: the current transaction must finish posting before
+				// it can be rolled back. The lambda is intentionally fire-and-forget with
+				// no captures from the surrounding stack, so no lifetime concerns arise.
+				// bCanRedo is false because the reverted parameter change left stale bindings
+				// that were declined — re-applying it would reintroduce the same stale state.
+				if (OutResponse == EStaleBindingResponse::Revert)
+				{
+					FTSTicker::GetCoreTicker().AddTicker(TEXT("SubsonicEventCollectionEditor::UndoParameterTransaction"), 0.0f, [](float)
+					{
+						if (GEditor)
+						{
+							GEditor->UndoTransaction(false /* bCanRedo */);
+						}
+						return false;
+					});
+				}
+			}
+		});
+	}
+
 	void FEventCollectionEditor::Init(const EToolkitMode::Type Mode, const TSharedPtr<IToolkitHost>& InitToolkitHost, USubsonicEventCollection& InEventCollection)
 	{
 		// Support undo/redo
@@ -769,6 +873,8 @@ namespace UE::Subsonic::Editor
 		USubsonicEventCollection* EventCollection = GetEventCollection();
 		check(EventCollection);
 
+		InitBindings();
+
 		RebuildEventTree();
 
 		if (ParametersDetailsView.IsValid())
@@ -930,9 +1036,7 @@ namespace UE::Subsonic::Editor
 	{
 		if (USubsonicEventCollection* Collection = GetEventCollection())
 		{
-			FScopedTransaction Transaction(Description);
-			Collection->Modify();
-			TransactionFunc(Collection->GetMutableDefinition());
+			UE::Subsonic::Editor::TransactEventCollection(Description, *Collection, TransactionFunc);
 		}
 	}
 
@@ -1169,6 +1273,14 @@ namespace UE::Subsonic::Editor
 		FEditorUndoClient::PostUndo(bSuccess);
 	}
 
+	void FEventCollectionEditor::NotifyPostChange(const FPropertyChangedEvent& PropertyChangedEvent, FProperty* PropertyThatChanged)
+	{
+		if (EventTreeInspectorView.IsValid())
+		{
+			EventTreeInspectorView->ForceRefresh();
+		}
+	}
+
 	void FEventCollectionEditor::RestartAudition(bool bAutoEventsEnabled)
 	{
 		if (const USubsonicEventCollection* EventCollection = GetEventCollection())
diff --git a/Engine/Plugins/Experimental/Subsonic/Source/SubsonicEditor/Private/SubsonicEventCollectionEditor.h b/Engine/Plugins/Experimental/Subsonic/Source/SubsonicEditor/Private/SubsonicEventCollectionEditor.h
--- a/Engine/Plugins/Experimental/Subsonic/Source/SubsonicEditor/Private/SubsonicEventCollectionEditor.h
+++ b/Engine/Plugins/Experimental/Subsonic/Source/SubsonicEditor/Private/SubsonicEventCollectionEditor.h
@@ -42,6 +42,11 @@ namespace UE::Subsonic
 
 			USubsonicEventCollection* GetEventCollection() const;
 
+			// Selects the event row (not an action row) in the event tree for the given event name.
+			void SelectEventInTree(FName EventName);
+
+			void TryInvokeTab(FName TabName) const;
+
 			/** IToolkit Interface */
 			virtual FName GetToolkitFName() const override;
 			virtual FText GetBaseToolkitName() const override;
@@ -57,6 +62,9 @@ namespace UE::Subsonic
 			virtual void PostUndo(bool bSuccess) override;
 			virtual void PostRedo(bool bSuccess) override { PostUndo(bSuccess); }
 
+			/** FNotifyHook Interface */
+			virtual void NotifyPostChange(const FPropertyChangedEvent& PropertyChangedEvent, FProperty* PropertyThatChanged) override;
+
 		private:
 			void TransactEvent(const FEventHandle EventHandle, FText Description, TFunctionRef<void(FSubsonicEvent&)> TransactionFunc) const;
 			void TransactEventCollection(FText Description, TFunctionRef<void(FSubsonicEventCollectionDefinition&)> TransactionFunc) const;
@@ -109,6 +117,8 @@ namespace UE::Subsonic
 
 			void ExecutePreviewEvent(const FEventHandle& EventHandle);
 
+			void InitBindings();
+
 			using SHandleVariantTableRow = STableRow<TSharedPtr<FHandleVariant>>;
 
 			// Currently, only supports actions. Could support dragging and ordering events in specialized view sorted by user-defined index in the future.
@@ -146,6 +156,8 @@ namespace UE::Subsonic
 			// Single property view of collection for use to capture bespoke property handles
 			// for visualization in contexts outside of a contextual details layout/panel
 			TSharedPtr<ISinglePropertyView> CollectionView;
+
+			FDelegateHandle StaleBindingsDelegateHandle;
 		};
 	} // namespace Editor
 } // namespace UE::Subsonic
\ No newline at end of file
diff --git a/Engine/Plugins/Experimental/Subsonic/Source/SubsonicEditor/Private/SubsonicPropertyBindingExtension.cpp b/Engine/Plugins/Experimental/Subsonic/Source/SubsonicEditor/Private/SubsonicPropertyBindingExtension.cpp
new file mode 100644
--- /dev/null
+++ b/Engine/Plugins/Experimental/Subsonic/Source/SubsonicEditor/Private/SubsonicPropertyBindingExtension.cpp
@@ -0,0 +1,526 @@
+// Copyright Epic Games, Inc. All Rights Reserved.
+
+#include "SubsonicPropertyBindingExtension.h"
+
+#include "Algo/AnyOf.h"
+#include "EdGraphSchema_K2.h"
+#include "Features/IModularFeatures.h"
+#include "Framework/Application/SlateApplication.h"
+#include "Framework/MultiBox/MultiBoxBuilder.h"
+#include "SubsonicBindingUtils.h"
+#include "SubsonicEventCollection.h"
+#include "SubsonicEventCollectionViews.h"
+#include "Widgets/Input/SCheckBox.h"
+#include "Widgets/Input/SHyperlink.h"
+#include "Widgets/Layout/SBox.h"
+#include "Widgets/SToolTip.h"
+
+#define LOCTEXT_NAMESPACE "SubsonicEditor"
+
+namespace UE::Subsonic::Editor
+{
+	FSubsonicPropertyBindingExtension::FSubsonicPropertyBindingExtension(FNavigateToPropertyOwnerFunc InNavigateToPropertyOwner)
+		: NavigateToPropertyOwner(MoveTemp(InNavigateToPropertyOwner))
+	{
+	}
+
+	bool FSubsonicPropertyBindingExtension::IsPropertyExtendable(const UClass* InObjectClass, const IPropertyHandle& InPropertyHandle) const
+	{
+		const FProperty* Property = InPropertyHandle.GetProperty();
+		if (!Property || Property->HasMetaData("NoBinding"))
+		{
+			return false;
+		}
+
+		// Skip the base class struct-ID check; we supply our own binding context below.
+		return true;
+	}
+
+	FPropertyBindingWidgetArgs FSubsonicPropertyBindingExtension::InitWidgetRow(
+		USubsonicEventCollection& Collection,
+		USubsonicEventTreeDetailsView& View,
+		const FProperty* TargetProperty,
+		FDetailWidgetRow& WidgetRow)
+	{
+		check(TargetProperty);
+
+		// Weak refs used inside all delegates to guard against stale pointers.
+		TWeakObjectPtr<USubsonicEventCollection> WeakCollection(&Collection);
+		TWeakObjectPtr<USubsonicEventTreeDetailsView> WeakView(&View);
+
+		// Returns which parameter name (if any) is currently bound to PropertyName.
+		auto GetCurrentBindingParam = [WeakCollection, WeakView, TargetProperty]()
+		{
+			const USubsonicEventCollection* Collection = WeakCollection.Get();
+			const USubsonicEventTreeDetailsView* DetailsView = WeakView.Get();
+			if (Collection && DetailsView)
+			{
+				if (const FSubsonicEvent* Event = Collection->GetDefinition().FindEvent(DetailsView->Event))
+				{
+					const int32 Index = DetailsView->ActionIndex;
+					const TArray<FSubsonicEventActionDefinition>& ActionCollection = Event->GetActionCollection();
+					if (ActionCollection.IsValidIndex(Index))
+					{
+						if (const FSubsonicEventActionBase* ActionBase = ActionCollection[Index].GetAction().GetPtr())
+						{
+							const FName PropertyName = TargetProperty->GetFName();
+							return ActionBase->GetBoundParameterForProperty(PropertyName);
+						}
+					}
+				}
+			}
+			return FName();
+		};
+
+		TSharedRef<SWidget> ValueWidget = WidgetRow.ValueContent().Widget;
+		WidgetRow.ValueContent()
+		[
+			SNew(SBox)
+			.IsEnabled_Lambda([GetCurrentBindingParam]()
+			{
+				return GetCurrentBindingParam().IsNone();
+			})
+			.ToolTip_Lambda([GetCurrentBindingParam, ValueWidget]()
+			{
+				const FName CurrentParam = GetCurrentBindingParam();
+				return CurrentParam.IsNone()
+					? ValueWidget->GetToolTip()
+					: SNew(SToolTip).Text(FText::Format(
+						LOCTEXT("NotEditableWhenBound_Tooltip", "Parameter is currently bound to {0} (Remove binding to directly edit)"),
+						FText::FromName(CurrentParam)));
+			})
+			[
+				ValueWidget
+			]
+		];
+
+		FPropertyBindingWidgetArgs Args;
+		// FPropertyBindingWidgetArgs::Property is not const-correct; the binding widget only reads the property.
+		Args.Property = const_cast<FProperty*>(TargetProperty);
+		Args.bAllowNewBindings = false;
+		Args.bAllowFunctionBindings = false;
+		Args.bAllowUObjectFunctions = false;
+		Args.bAllowStructMemberBindings = false;
+		Args.BindButtonStyle = &FAppStyle::Get().GetWidgetStyle<FButtonStyle>("HoverHintOnly");
+
+		Args.OnGenerateMenu = FOnGetContent::CreateLambda(
+			[this, TargetProperty, GetCurrentBindingParam, WeakCollection, WeakView]()
+			{
+				return GenerateBindingMenu(TargetProperty, GetCurrentBindingParam, WeakCollection, WeakView);
+			});
+
+		Args.OnRemoveBinding = FOnRemoveBinding::CreateLambda(
+			[WeakCollection, WeakView, TargetProperty](FName /*InPropertyName*/)
+			{
+				RemoveBinding(WeakCollection, WeakView, TargetProperty);
+			});
+
+		Args.OnCanRemoveBinding = FOnCanRemoveBinding::CreateLambda([GetCurrentBindingParam](FName)
+		{
+			return !GetCurrentBindingParam().IsNone();
+		});
+
+		Args.OnHasAnyBindings = FOnHasAnyBindings::CreateLambda([GetCurrentBindingParam]()
+		{
+			return !GetCurrentBindingParam().IsNone();
+		});
+
+		Args.CurrentBindingText = TAttribute<FText>::Create([GetCurrentBindingParam, WeakCollection, WeakView]()
+		{
+			return GetBindingText(GetCurrentBindingParam, WeakCollection, WeakView);
+		});
+
+		return Args;
+	}
+
+	TSharedRef<SWidget> FSubsonicPropertyBindingExtension::GenerateBindingMenu(
+		const FProperty* TargetProperty,
+		TFunction<FName()> GetCurrentBindingParam,
+		TWeakObjectPtr<USubsonicEventCollection> WeakCollection,
+		TWeakObjectPtr<USubsonicEventTreeDetailsView> WeakView)
+	{
+		constexpr bool bInShouldCloseWindowAfterMenuSelection = true;
+		TSharedPtr<const FUICommandList> NullCommandList;
+		FMenuBuilder MenuBuilder(bInShouldCloseWindowAfterMenuSelection, NullCommandList);
+
+		const USubsonicEventCollection* Collection = WeakCollection.Get();
+		const USubsonicEventTreeDetailsView* DetailsView = WeakView.Get();
+		if (!Collection || !DetailsView)
+		{
+			return MenuBuilder.MakeWidget();
+		}
+
+		// Rebuild binding contexts from live data each time the menu opens so that
+		// renamed / removed parameters are never shown with stale names.
+		constexpr FSlateBrush* NullIconBrush = nullptr;
+		TArray<FSubsonicBindingContextStruct> BindingContextStructs;
+		const FSubsonicEventCollectionDefinition& CollectionDef = Collection->GetDefinition();
+
+		if (const UPropertyBag* CollectionBag = CollectionDef.GetParameters().GetPropertyBagStruct())
+		{
+			const FText Desc = LOCTEXT("BindingContext_CollectionDescription", "Collection");
+			BindingContextStructs.Add(FSubsonicBindingContextStruct(CollectionBag, NullIconBrush, Desc));
+		}
+
+		if (const FSubsonicEvent* Event = CollectionDef.FindEvent(DetailsView->Event))
+		{
+			if (const UPropertyBag* EventBag = Event->GetParameters().GetPropertyBagStruct())
+			{
+				const FText Desc = LOCTEXT("BindingContext_EventDescription", "Event");
+				FSubsonicBindingContextStruct Context(EventBag, NullIconBrush, Desc);
+				Context.bIsEventBinding = true;
+				BindingContextStructs.Add(MoveTemp(Context));
+			}
+		}
+
+		struct FStructBindingEntry
+		{
+			int32 Index = INDEX_NONE;
+			FProperty* BagProp = nullptr;
+		};
+
+		TArray<FName> OrderedNames;
+		TMap<FName, FStructBindingEntry> EntryMap;
+
+		for (int32 Index = 0; Index < BindingContextStructs.Num(); ++Index)
+		{
+			const FBindingContextStruct& BindingContext = BindingContextStructs[Index];
+			for (TFieldIterator<FProperty> PropIt(BindingContext.Struct); PropIt; ++PropIt)
+			{
+				FProperty* BagProp = *PropIt;
+				if (!BagProp || !BindingUtils::ArePropertiesBindingCompatible(BagProp, TargetProperty))
+				{
+					continue;
+				}
+				const FName BagPropName = BagProp->GetFName();
+				if (!EntryMap.Contains(BagPropName))
+				{
+					OrderedNames.Add(BagPropName);
+				}
+
+				// Overwrite any earlier context entry with event context (Event > Collection).
+				EntryMap.Emplace(BagPropName, FStructBindingEntry { Index, BagProp });
+			}
+		}
+
+		Algo::Sort(OrderedNames, [](const FName& A, const FName& B) { return A.LexicalLess(B); });
+
+		for (const FName& BagPropName : OrderedNames)
+		{
+			const FStructBindingEntry& Entry = EntryMap[BagPropName];
+			const int32 Index = Entry.Index;
+			const FSubsonicBindingContextStruct& BindingContext = BindingContextStructs[Index];
+			FProperty* BagProp = Entry.BagProp;
+			const FText Tooltip = FText::Format(LOCTEXT("BindingSourceLevel_Tooltip", "Defined in: {0}"),
+				BindingContext.DisplayText);
+
+			// For contexts that have a navigation target, the tooltip is an interactive
+			// hyperlink that selects the source in the editor when clicked. For other contexts
+			// (Collection) it is plain text.
+			TSharedRef<FCheckBoxStyle> EntryStyle = GetCheckboxEntryStyle(*BagProp);
+			TSharedRef<SToolTip> EntryTooltip = SNew(SToolTip)
+				.IsInteractive(true)
+				.BorderImage(FCoreStyle::Get().GetBrush("ToolTip.BrightBackground"))
+				.TextMargin(FMargin(11.0f))
+				.Font(FAppStyle::GetFontStyle("PropertyWindow.NormalFont"))
+				[
+					SNew(SHyperlink)
+					.Text(Tooltip)
+					.Style(FAppStyle::Get(), "DarkHyperlink")
+					.OnNavigate_Lambda([NavigateToPropertyOwner = NavigateToPropertyOwner, bIsEvent = BindingContext.bIsEventBinding, WeakView]()
+					{
+						if (const USubsonicEventTreeDetailsView* View = WeakView.Get())
+						{
+							const FName BagEventName = bIsEvent ? View->Event : FName();
+							NavigateToPropertyOwner(BagEventName);
+						}
+					})
+				];
+
+			MenuBuilder.AddWidget(
+				SNew(SBox)
+				.Padding(FMargin(30.0f, 2.0f, 4.0f, 4.0f))
+				[
+					SNew(SCheckBox)
+					.Style(&EntryStyle.Get())
+					.Padding(FMargin(8.0f, 0.0f, 0.0f, 0.0f))
+					.ToolTip(EntryTooltip)
+					.IsChecked_Lambda([GetCurrentBindingParam, BagPropName, EntryStyle]()
+					{
+						return GetCurrentBindingParam() == BagPropName
+							? ECheckBoxState::Checked
+							: ECheckBoxState::Unchecked;
+					})
+					.OnCheckStateChanged_Lambda([WeakCollection, WeakView, TargetProperty, Index, BagProp, GetCurrentBindingParam, BagPropName](ECheckBoxState)
+					{
+						// Always close the menu on click (radio button semantics).
+						// Only re-bind if selecting a different parameter to avoid a spurious undo entry.
+						if (GetCurrentBindingParam() != BagPropName)
+						{
+							AddParamBinding(WeakCollection.Get(), WeakView.Get(), TargetProperty->GetFName(), Index, BagProp);
+						}
+						FSlateApplication::Get().DismissAllMenus();
+					})
+					[
+						SNew(STextBlock)
+						.TextStyle(FAppStyle::Get(), "Menu.Label")
+						.Text(FText::FromName(BagPropName))
+					]
+				],
+				FText::GetEmpty()
+			);
+		}
+
+		const FName BoundParam = GetCurrentBindingParam();
+		if (!BoundParam.IsNone())
+		{
+			MenuBuilder.AddMenuSeparator();
+			MenuBuilder.AddMenuEntry(
+				LOCTEXT("RemoveBinding", "Remove Binding"),
+				FText::GetEmpty(),
+				FSlateIcon(FAppStyle::GetAppStyleSetName(), "Icons.X"),
+				FUIAction(FExecuteAction::CreateLambda(
+					[WeakCollection, WeakView, TargetProperty]()
+					{
+						RemoveBinding(WeakCollection, WeakView, TargetProperty);
+					})
+				)
+			);
+		}
+		return MenuBuilder.MakeWidget();
+	}
+
+	FText FSubsonicPropertyBindingExtension::GetBindingText(
+		TFunction<FName()> GetCurrentBindingParam,
+		TWeakObjectPtr<USubsonicEventCollection> WeakCollection,
+		TWeakObjectPtr<USubsonicEventTreeDetailsView> WeakView)
+	{
+		const FName BoundParam = GetCurrentBindingParam();
+		if (BoundParam.IsNone())
+		{
+			return FText();
+		}
+
+		// Resolve from live data to avoid stale bag pointers. Check Event first
+		// since event-level parameters take precedence over collection-level.
+		const USubsonicEventCollection* Collection = WeakCollection.Get();
+		const USubsonicEventTreeDetailsView* DetailsView = WeakView.Get();
+		if (Collection && DetailsView)
+		{
+			const FSubsonicEventCollectionDefinition& Def = Collection->GetDefinition();
+			if (const FSubsonicEvent* Event = Def.FindEvent(DetailsView->Event))
+			{
+				if (const UPropertyBag* EventBag = Event->GetParameters().GetPropertyBagStruct())
+				{
+					if (EventBag->FindPropertyByName(BoundParam))
+					{
+						return FText::Format(LOCTEXT("BoundParameter_DescFormat", "{0} > {1}"),
+							LOCTEXT("BindingContext_EventDescription", "Event"), FText::FromName(BoundParam));
+					}
+				}
+			}
+			if (const UPropertyBag* CollBag = Def.GetParameters().GetPropertyBagStruct())
+			{
+				if (CollBag->FindPropertyByName(BoundParam))
+				{
+					return FText::Format(LOCTEXT("BoundParameter_DescFormat", "{0} > {1}"),
+						LOCTEXT("BindingContext_CollectionDescription", "Collection"), FText::FromName(BoundParam));
+				}
+			}
+		}
+		return FText::FromName(BoundParam);
+	}
+
+	void FSubsonicPropertyBindingExtension::RemoveBinding(
+		TWeakObjectPtr<USubsonicEventCollection> WeakCollection,
+		TWeakObjectPtr<USubsonicEventTreeDetailsView> WeakView,
+		const FProperty* TargetProperty)
+	{
+		USubsonicEventCollection* Collection = WeakCollection.Get();
+		USubsonicEventTreeDetailsView* DetailsView = WeakView.Get();
+		if (Collection && DetailsView)
+		{
+			const int32 Index = DetailsView->ActionIndex;
+			const FEventHandle EventHandle { .Collection = Collection->GetHandle(), .EventName = DetailsView->Event };
+			FSubsonicEvent* Event = Collection->GetMutableDefinition().FindMutableEvent(EventHandle);
+			if (Event && Event->GetActionCollection().IsValidIndex(Index))
+			{
+				if (FSubsonicEventActionBase* ActionBase = Event->GetMutableActionCollection()[Index].GetMutableActionBase())
+				{
+					TransactEventCollection(LOCTEXT("RemoveSubsonicBinding", "Remove Subsonic Parameter Binding"), *Collection, [&](FSubsonicEventCollectionDefinition&)
+					{
+						const FName PropertyName = TargetProperty->GetFName();
+						ActionBase->RemovePropertyBinding(PropertyName);
+					});
+				}
+			}
+		}
+	}
+
+	void FSubsonicPropertyBindingExtension::ExtendWidgetRow(FDetailWidgetRow& WidgetRow, const IDetailLayoutBuilder& InDetailBuilder, const UClass* InObjectClass, TSharedPtr<IPropertyHandle> InPropertyHandle)
+	{
+		if (!IModularFeatures::Get().IsModularFeatureAvailable("PropertyAccessEditor"))
+		{
+			return;
+		}
+
+		// Resolve the inspector view so we know which event / action is selected.
+		TArray<TWeakObjectPtr<UObject>> Objects;
+		InDetailBuilder.GetObjectsBeingCustomized(Objects);
+		if (Objects.Num() != 1)
+		{
+			return;
+		}
+
+		USubsonicEventTreeDetailsView* View = Cast<USubsonicEventTreeDetailsView>(Objects.Last().Get());
+		if (!View || View->ActionIndex < 0)
+		{
+			return;
+		}
+
+		USubsonicEventCollection* Collection = View->GetCollection();
+		if (!Collection)
+		{
+			return;
+		}
+
+		const FProperty* TargetProperty = InPropertyHandle->GetProperty();
+		if (!TargetProperty)
+		{
+			return;
+		}
+
+		TArray<FBindingContextStruct> BindingContextStructs;
+		const FSubsonicEventCollectionDefinition& CollectionDef = Collection->GetDefinition();
+		constexpr FSlateBrush* NullIconBrush = nullptr;
+
+		if (const UPropertyBag* CollectionBag = CollectionDef.GetParameters().GetPropertyBagStruct())
+		{
+			const FText BindingContextDesc = LOCTEXT("BindingContext_CollectionDescription", "Collection");
+			BindingContextStructs.Add(FSubsonicBindingContextStruct(CollectionBag, NullIconBrush, BindingContextDesc));
+		}
+
+		if (const FSubsonicEvent* Event = CollectionDef.FindEvent(View->Event))
+		{
+			if (const UPropertyBag* EventBag = Event->GetParameters().GetPropertyBagStruct())
+			{
+				const FText BindingContextDesc = LOCTEXT("BindingContext_EventDescription", "Event");
+				BindingContextStructs.Add(FSubsonicBindingContextStruct(EventBag, NullIconBrush, BindingContextDesc));
+			}
+		}
+
+		// Only show the binding dropdown if at least one compatible parameter exists.
+		const bool bHasCompatibleParam = Algo::AnyOf(BindingContextStructs, [&TargetProperty](const FBindingContextStruct& Context)
+		{
+			for (TFieldIterator<FProperty> PropIt(Context.Struct); PropIt; ++PropIt)
+			{
+				if (*PropIt && BindingUtils::ArePropertiesBindingCompatible(*PropIt, TargetProperty))
+				{
+					return true;
+				}
+			}
+			return false;
+		});
+		if (bHasCompatibleParam)
+		{
+			FPropertyBindingWidgetArgs Args = InitWidgetRow(*Collection, *View, TargetProperty, WidgetRow);
+			IPropertyAccessEditor& PropertyAccessEditor = IModularFeatures::Get().GetModularFeature<IPropertyAccessEditor>("PropertyAccessEditor");
+			WidgetRow.ExtensionContent()
+			[
+				PropertyAccessEditor.MakePropertyBindingWidget(BindingContextStructs, Args)
+			];
+		}
+	}
+
+	void FSubsonicPropertyBindingExtension::AddParamBinding(USubsonicEventCollection* Collection, USubsonicEventTreeDetailsView* DetailsView, FName PropertyName, const int32 SourceStructIndex, const FProperty* SourceProperty)
+	{
+		if (!SourceProperty || !Collection || !DetailsView || DetailsView->ActionIndex < 0)
+		{
+			return;
+		}
+
+		const FEventHandle EventHandle { .Collection = Collection->GetHandle(), .EventName = DetailsView->Event };
+		FSubsonicEvent* Event = Collection->GetMutableDefinition().FindMutableEvent(EventHandle);
+		if (!Event || !Event->GetActionCollection().IsValidIndex(DetailsView->ActionIndex))
+		{
+			return;
+		}
+
+		FSubsonicEventActionDefinition& ActionDef = Event->GetMutableActionCollection()[DetailsView->ActionIndex];
+		FSubsonicEventActionBase* ActionBase = ActionDef.GetMutableActionBase();
+		if (!ActionBase)
+		{
+			return;
+		}
+
+		const FName SourceParamName = SourceProperty->GetFName();
+
+		TransactEventCollection(LOCTEXT("AddSubsonicBinding", "Add Subsonic Parameter Binding"), *Collection, [&](FSubsonicEventCollectionDefinition& CollectionDef)
+		{
+			ActionBase->AddPropertyBinding(SourceParamName, PropertyName);
+
+			const UPropertyBag* CollectionBag = CollectionDef.GetParameters().GetPropertyBagStruct();
+			const FInstancedPropertyBag* SourceBag = SourceStructIndex == CollectionContextIndex && CollectionBag != nullptr
+				? &CollectionDef.GetParameters()
+				: &Event->GetParameters();
+
+			if (!SourceBag)
+			{
+				return;
+			}
+
+			// Re-resolve the source property from the current bag to avoid using a stale FProperty*
+			// captured at menu creation time (the PropertyBag struct may have been rebuilt since then).
+			const UPropertyBag* BagStruct = SourceBag->GetPropertyBagStruct();
+			const FPropertyBagPropertyDesc* BagDesc = BagStruct ? BagStruct->FindPropertyDescByName(SourceParamName) : nullptr;
+			if (!BagDesc || !BagDesc->CachedProperty)
+			{
+				return;
+			}
+
+			const FConstStructView BagView = SourceBag->GetValue();
+			if (!BagView.IsValid())
+			{
+				return;
+			}
+
+			const UScriptStruct* ActionStruct = ActionDef.GetAction().GetScriptStruct();
+			if (!ActionStruct)
+			{
+				return;
+			}
+
+			for (TFieldIterator<FProperty> PropIt(ActionStruct); PropIt; ++PropIt)
+			{
+				if (PropIt->GetFName() == PropertyName && BindingUtils::ArePropertiesBindingCompatible(BagDesc->CachedProperty , *PropIt))
+				{
+					const void* SrcPtr = BagDesc->CachedProperty->ContainerPtrToValuePtr<const void>(BagView.GetMemory());
+					void* DstPtr = PropIt->ContainerPtrToValuePtr<void>(ActionBase);
+					PropIt->CopyCompleteValue(DstPtr, SrcPtr);
+					break;
+				}
+			}
+		});
+	}
+
+	TSharedRef<FCheckBoxStyle> FSubsonicPropertyBindingExtension::GetCheckboxEntryStyle(const FProperty& BagProp)
+	{
+		FEdGraphPinType PinType;
+		FLinearColor TypeColor = FLinearColor::White;
+		if (const UEdGraphSchema_K2* K2Schema = GetDefault<UEdGraphSchema_K2>())
+		{
+			if (K2Schema->ConvertPropertyToPinType(&BagProp, PinType))
+			{
+				TypeColor = K2Schema->GetPinTypeColor(PinType);
+			}
+		}
+
+		TSharedRef<FCheckBoxStyle> EntryStyle =  MakeShared<FCheckBoxStyle>(FAppStyle::Get().GetWidgetStyle<FCheckBoxStyle>("Menu.RadioButton"));
+		EntryStyle->CheckedImage.TintColor = FSlateColor(TypeColor);
+		EntryStyle->CheckedHoveredImage.TintColor = FSlateColor(TypeColor);
+		EntryStyle->CheckedPressedImage.TintColor = FSlateColor(TypeColor);
+		return EntryStyle;
+	}
+} // namespace UE::Subsonic::Editor
+
+#undef LOCTEXT_NAMESPACE
diff --git a/Engine/Plugins/Experimental/Subsonic/Source/SubsonicEditor/Private/SubsonicPropertyBindingExtension.h b/Engine/Plugins/Experimental/Subsonic/Source/SubsonicEditor/Private/SubsonicPropertyBindingExtension.h
new file mode 100644
--- /dev/null
+++ b/Engine/Plugins/Experimental/Subsonic/Source/SubsonicEditor/Private/SubsonicPropertyBindingExtension.h
@@ -0,0 +1,97 @@
+// Copyright Epic Games, Inc. All Rights Reserved.
+
+#pragma once
+
+#include "IPropertyAccessEditor.h"
+#include "PropertyBindingExtension.h"
+#include "StructUtils/PropertyBag.h"
+
+
+namespace UE::Subsonic
+{
+	// Forward Declarations
+	class USubsonicEventCollection;
+	struct FSubsonicEventCollectionDefinition;
+
+	namespace Editor
+	{
+		// Forward Declarations
+		class USubsonicEventTreeDetailsView;
+
+		// Extends FBindingContextStruct to accept a const UPropertyBag* and track
+		// whether the binding originates from an event-level parameter bag.
+		// The const_cast in the constructor is required because FBindingContextStruct::Struct
+		// is a non-const UStruct* — the underlying property-access API is not const-correct.
+		struct FSubsonicBindingContextStruct : public FBindingContextStruct
+		{
+			bool bIsEventBinding = false;
+
+			FSubsonicBindingContextStruct(
+				const UPropertyBag* InBag,
+				const FSlateBrush* InIcon = nullptr,
+				const FText& InDisplayText = FText::GetEmpty(),
+				const FText& InTooltipText = FText::GetEmpty(),
+				const FText& InSection = FText::GetEmpty())
+				: FBindingContextStruct(const_cast<UPropertyBag*>(InBag), InIcon, InDisplayText, InTooltipText, InSection)
+			{
+			}
+		};
+
+		void TransactEventCollection(FText Description, USubsonicEventCollection& Collection, TFunctionRef<void(FSubsonicEventCollectionDefinition&)> TransactionFunc);
+
+		// Only shows the binding dropdown on properties not decorated with 'NoBinding' Metadata
+		// tag. Overrides ExtendWidgetRow to populate the binding menu with parameters from the
+		// collection- or event-level (Event value supersedes collection value if types match)
+		// Subsonic 'Parameter' property bags.
+		class FSubsonicPropertyBindingExtension : public FPropertyBindingExtension
+		{
+		public:
+			using FNavigateToPropertyOwnerFunc = TFunction<void(FName)>;
+
+			explicit FSubsonicPropertyBindingExtension(FNavigateToPropertyOwnerFunc InNavigateToPropertyOwner);
+
+			//~ Begin IDetailPropertyExtensionHandler
+			virtual bool IsPropertyExtendable(const UClass* InObjectClass, const IPropertyHandle& InPropertyHandle) const override;
+			virtual void ExtendWidgetRow(FDetailWidgetRow& WidgetRow, const IDetailLayoutBuilder& InDetailBuilder, const UClass* InObjectClass, TSharedPtr<IPropertyHandle> InPropertyHandle) override;
+			//~ End IDetailPropertyExtensionHandler
+
+		private:
+			// Index into the BindingContextStructs array built in GenerateBindingMenu / ExtendWidgetRow.
+			// Collection is always added first (index 0), event second (index 1).
+			static const int32 CollectionContextIndex = 0;
+
+			FPropertyBindingWidgetArgs InitWidgetRow(
+				USubsonicEventCollection& Collection,
+				USubsonicEventTreeDetailsView& View,
+				const FProperty* TargetProperty,
+				FDetailWidgetRow& WidgetRow);
+
+			TSharedRef<SWidget> GenerateBindingMenu(
+				const FProperty* TargetProperty,
+				TFunction<FName()> GetCurrentBindingParam,
+				TWeakObjectPtr<USubsonicEventCollection> WeakCollection,
+				TWeakObjectPtr<USubsonicEventTreeDetailsView> WeakView);
+
+			static FText GetBindingText(
+				TFunction<FName()> GetCurrentBindingParam,
+				TWeakObjectPtr<USubsonicEventCollection> WeakCollection,
+				TWeakObjectPtr<USubsonicEventTreeDetailsView> WeakView);
+
+			static void RemoveBinding(
+				TWeakObjectPtr<USubsonicEventCollection> WeakCollection,
+				TWeakObjectPtr<USubsonicEventTreeDetailsView> WeakView,
+				const FProperty* TargetProperty);
+
+			static void AddParamBinding(
+				USubsonicEventCollection* Collection,
+				USubsonicEventTreeDetailsView* DetailsView,
+				FName PropertyName,
+				const int32 SourceStructIndex,
+				const FProperty* SourceProperty);
+
+			static TSharedRef<FCheckBoxStyle> GetCheckboxEntryStyle(const FProperty& BagProp);
+
+			FNavigateToPropertyOwnerFunc NavigateToPropertyOwner;
+		};
+	} // namespace Editor
+} // namespace UE::Subsonic
diff --git a/Engine/Plugins/Experimental/Subsonic/Source/SubsonicEditor/SubsonicEditor.Build.cs b/Engine/Plugins/Experimental/Subsonic/Source/SubsonicEditor/SubsonicEditor.Build.cs
--- a/Engine/Plugins/Experimental/Subsonic/Source/SubsonicEditor/SubsonicEditor.Build.cs
+++ b/Engine/Plugins/Experimental/Subsonic/Source/SubsonicEditor/SubsonicEditor.Build.cs
@@ -22,6 +22,7 @@ public class SubsonicEditor : ModuleRules
 				"AssetDefinition",
 				"AudioEditor",
 				"AudioWidgets",
+				"BlueprintGraph",
 				"ContentBrowser",
 				"EditorSubsystem",
 				"Engine",
