Subsonic GeneratorSource parameter infrastructure and MetaSound support

#rnx

#Synced-CL 52256849

diff --git a/Engine/Plugins/Experimental/Subsonic/Source/Subsonic/Private/CoreActions/SubsonicAction_GeneratorSource.cpp b/Engine/Plugins/Experimental/Subsonic/Source/Subsonic/Private/CoreActions/SubsonicAction_GeneratorSource.cpp
--- a/Engine/Plugins/Experimental/Subsonic/Source/Subsonic/Private/CoreActions/SubsonicAction_GeneratorSource.cpp
+++ b/Engine/Plugins/Experimental/Subsonic/Source/Subsonic/Private/CoreActions/SubsonicAction_GeneratorSource.cpp
@@ -15,25 +15,29 @@ namespace UE::Subsonic
 			return;
 		}
 
+		// Merge authored params (lowest priority) with trigger-time params from the executor.
+		FSubsonicParameterStore MergedParams = Parameters;
+		MergedParams.MergeFrom(InExecutor.GetParameters());
+
 		const FActionScopeKey ScopeKey(InExecutor, InHandle);
 		if (USubsonicGeneratorSourceSubscriber* GeneratorSourceSubscriber = ScopeKey.FindDeviceSubsystem<USubsonicGeneratorSourceSubscriber>())
 		{
 			switch (Scope)
 			{
 			case ESubsonicExecutionScope::Global:
-				GeneratorSourceSubscriber->PlaySound(Name, *Sound);
+				GeneratorSourceSubscriber->PlaySound(Name, *Sound, MergedParams);
 				break;
 
 			case ESubsonicExecutionScope::Collection:
-				GeneratorSourceSubscriber->PlaySound(InHandle.Event.Collection, Name, *Sound);
+				GeneratorSourceSubscriber->PlaySound(InHandle.Event.Collection, Name, *Sound, MergedParams);
 				break;
 
 			case ESubsonicExecutionScope::Event:
-				GeneratorSourceSubscriber->PlaySound(InHandle.Event, Name, *Sound);
+				GeneratorSourceSubscriber->PlaySound(InHandle.Event, Name, *Sound, MergedParams);
 				break;
 
 			case ESubsonicExecutionScope::Action:
-				GeneratorSourceSubscriber->PlaySound(InHandle, Name, *Sound);
+				GeneratorSourceSubscriber->PlaySound(InHandle, Name, *Sound, MergedParams);
 				break;
 
 			default:
@@ -70,4 +74,4 @@ namespace UE::Subsonic
 			}
 		}
 	}
-}
\ No newline at end of file
+}
diff --git a/Engine/Plugins/Experimental/Subsonic/Source/Subsonic/Private/CoreActions/SubsonicAction_GeneratorSource.h b/Engine/Plugins/Experimental/Subsonic/Source/Subsonic/Private/CoreActions/SubsonicAction_GeneratorSource.h
--- a/Engine/Plugins/Experimental/Subsonic/Source/Subsonic/Private/CoreActions/SubsonicAction_GeneratorSource.h
+++ b/Engine/Plugins/Experimental/Subsonic/Source/Subsonic/Private/CoreActions/SubsonicAction_GeneratorSource.h
@@ -4,6 +4,7 @@
 #include "Misc/CoreMiscDefines.h"
 #include "Sound/SoundWave.h"
 #include "SubsonicEventCollection.h"
+#include "SubsonicParameterStore.h"
 
 #include "SubsonicAction_GeneratorSource.generated.h"
 
@@ -32,6 +33,10 @@ namespace UE::Subsonic
 		// within the collection, etc.
 		UPROPERTY(EditAnywhere, Category = "Subsonic")
 		ESubsonicExecutionScope Scope = ESubsonicExecutionScope::Collection;
+
+		// Authored parameters applied to the generator at play time (lowest priority layer).
+		UPROPERTY(EditAnywhere, Category = "Parameters")
+		FSubsonicParameterStore Parameters;
 	};
 
 	USTRUCT(MinimalAPI, BlueprintType, DisplayName = "Stop Sound (GeneratorSource)")
@@ -54,4 +59,4 @@ namespace UE::Subsonic
 		UPROPERTY(EditAnywhere, Category = "Subsonic")
 		ESubsonicExecutionScope Scope = ESubsonicExecutionScope::Collection;
 	};
-}
\ No newline at end of file
+}
diff --git a/Engine/Plugins/Experimental/Subsonic/Source/Subsonic/Private/CoreEventSubscribers/SubsonicGenerator.cpp b/Engine/Plugins/Experimental/Subsonic/Source/Subsonic/Private/CoreEventSubscribers/SubsonicGenerator.cpp
new file mode 100644
--- /dev/null
+++ b/Engine/Plugins/Experimental/Subsonic/Source/Subsonic/Private/CoreEventSubscribers/SubsonicGenerator.cpp
@@ -0,0 +1,407 @@
+// Copyright Epic Games, Inc. All Rights Reserved.
+
+#include "SubsonicGenerator.h"
+
+#include "CoreEventSubscribers/SubsonicWaveGenerator.h"
+#include "DSP/Dsp.h"
+#include "DSP/FloatArrayMath.h"
+#include "IAudioMixerGeneratorSource.h"
+#include "StructUtils/PropertyBag.h"
+#include "SubsonicBuiltInParameters.h"
+#include "SubsonicLog.h"
+#include "SubsonicParameterStore.h"
+
+namespace UE::Subsonic
+{
+	const TMap<FName, FSubsonicGenerator::FBuiltInParamHandler>& FSubsonicGenerator::GetBuiltInParams()
+	{
+		static const TMap<FName, FBuiltInParamHandler> Map = {
+			{ BuiltInParameters::Volume,         &FSubsonicGenerator::SetVolume },
+			{ BuiltInParameters::PitchShift,     &FSubsonicGenerator::SetPitchShift },
+			{ BuiltInParameters::HighpassCutoff, &FSubsonicGenerator::SetHighpassCutoff },
+			{ BuiltInParameters::LowpassCutoff,  &FSubsonicGenerator::SetLowpassCutoff },
+			{ BuiltInParameters::FadeOutTime,    &FSubsonicGenerator::SetFadeOutTime },
+		};
+		return Map;
+	}
+
+	FSubsonicGenerator::FSubsonicGenerator(TSharedRef<FWaveGenerator> InWave, float InSampleRate)
+		: InnerGenerator(InWave)
+		, bGraphReady(true) // Wave path: no graph to wait for.
+		, SampleRate(InSampleRate)
+	{
+		InitDSP();
+	}
+
+	FSubsonicGenerator::FSubsonicGenerator(ISoundGeneratorPtr InMetaSoundGenerator, float InSampleRate)
+		: InnerGenerator(InMetaSoundGenerator)
+		, bGraphReady(false) // Will be set true by the graph-set callback.
+		, SampleRate(InSampleRate)
+	{
+		TSharedPtr<Metasound::FMetasoundGenerator, ESPMode::ThreadSafe> MSGen =
+			StaticCastSharedPtr<Metasound::FMetasoundGenerator>(InMetaSoundGenerator);
+		if (MSGen.IsValid())
+		{
+			MetaSoundGeneratorWeak = MSGen;
+
+			// Register for notification when the MetaSound graph is compiled and ready.
+			// If the graph is already built, AddGraphSetCallback fires immediately.
+			GraphSetCallbackHandle = MSGen->AddGraphSetCallback(
+				Metasound::FOnSetGraph::FDelegate::CreateLambda([this]()
+				{
+					bGraphReady.store(true, std::memory_order_release);
+				}));
+		}
+		else
+		{
+			bGraphReady.store(true, std::memory_order_relaxed);
+		}
+
+		InitDSP();
+	}
+
+	FSubsonicGenerator::~FSubsonicGenerator()
+	{
+		if (GraphSetCallbackHandle.IsValid())
+		{
+			if (TSharedPtr<Metasound::FMetasoundGenerator, ESPMode::ThreadSafe> MSGen = MetaSoundGeneratorWeak.Pin())
+			{
+				MSGen->RemoveGraphSetCallback(GraphSetCallbackHandle);
+			}
+		}
+	}
+
+	void FSubsonicGenerator::InitDSP()
+	{
+		const int32 NumCh = InnerGenerator->GetNumChannels();
+
+		HighpassFilter.Init(SampleRate, NumCh, Audio::EBiquadFilter::Highpass, 0.0f);
+		HighpassFilter.SetEnabled(false);
+
+		LowpassFilter.Init(SampleRate, NumCh, Audio::EBiquadFilter::Lowpass, SampleRate * 0.5f);
+		LowpassFilter.SetEnabled(false);
+
+		Resampler = MakeUnique<Audio::FRuntimeResampler>(NumCh);
+		Resampler->SetFrameRatio(1.0f);
+	}
+
+	bool FSubsonicGenerator::TryGetFloat(const FSubsonicParameterStore& Store, const FPropertyBagPropertyDesc& Desc, float& OutValue)
+	{
+		if (Desc.ValueType == EPropertyBagPropertyType::Float)
+		{
+			TValueOrError<float, EPropertyBagResult> Result = Store.Bag.GetValueFloat(Desc);
+			if (Result.HasValue())
+			{
+				OutValue = Result.GetValue();
+				return true;
+			}
+		}
+		else if (Desc.ValueType == EPropertyBagPropertyType::Double)
+		{
+			TValueOrError<double, EPropertyBagResult> Result = Store.Bag.GetValueDouble(Desc);
+			if (Result.HasValue())
+			{
+				OutValue = static_cast<float>(Result.GetValue());
+				return true;
+			}
+		}
+		return false;
+	}
+
+	bool FSubsonicGenerator::HasMetaSoundInput(FName Name) const
+	{
+		TSharedPtr<Metasound::FMetasoundGenerator, ESPMode::ThreadSafe> MSGen = MetaSoundGeneratorWeak.Pin();
+		if (!MSGen.IsValid())
+		{
+			return false;
+		}
+		return MSGen->GetInputWriteReference<float>(Name).IsSet();
+	}
+
+	void FSubsonicGenerator::ForwardToMetaSound(
+		Metasound::FMetasoundGenerator& MSGen,
+		const FSubsonicParameterStore& Store,
+		const FPropertyBagPropertyDesc& Desc)
+	{
+		switch (Desc.ValueType)
+		{
+		case EPropertyBagPropertyType::Float:
+		{
+			TValueOrError<float, EPropertyBagResult> Result = Store.Bag.GetValueFloat(Desc);
+			if (Result.HasValue())
+			{
+				MSGen.SetInputValue<float>(Desc.Name, Result.GetValue());
+			}
+			break;
+		}
+		case EPropertyBagPropertyType::Double:
+		{
+			TValueOrError<double, EPropertyBagResult> Result = Store.Bag.GetValueDouble(Desc);
+			if (Result.HasValue())
+			{
+				MSGen.SetInputValue<float>(Desc.Name, static_cast<float>(Result.GetValue()));
+			}
+			break;
+		}
+		case EPropertyBagPropertyType::Bool:
+		{
+			TValueOrError<bool, EPropertyBagResult> Result = Store.Bag.GetValueBool(Desc);
+			if (Result.HasValue())
+			{
+				MSGen.SetInputValue<bool>(Desc.Name, Result.GetValue());
+			}
+			break;
+		}
+		case EPropertyBagPropertyType::Int32:
+		{
+			TValueOrError<int32, EPropertyBagResult> Result = Store.Bag.GetValueInt32(Desc);
+			if (Result.HasValue())
+			{
+				MSGen.SetInputValue<int32>(Desc.Name, Result.GetValue());
+			}
+			break;
+		}
+		default:
+			break;
+		}
+	}
+
+	void FSubsonicGenerator::ApplyParameters(const FSubsonicParameterStore& Store)
+	{
+		const UPropertyBag* BagStruct = Store.Bag.GetPropertyBagStruct();
+		if (!BagStruct)
+		{
+			return;
+		}
+
+		// If the MetaSound graph isn't ready yet, defer the entire store.
+		if (!bGraphReady.load(std::memory_order_acquire))
+		{
+			DeferredParams = MakeShared<FSubsonicParameterStore>(Store);
+			return;
+		}
+
+		TSharedPtr<Metasound::FMetasoundGenerator, ESPMode::ThreadSafe> MSGen = MetaSoundGeneratorWeak.Pin();
+		const TMap<FName, FBuiltInParamHandler>& BuiltInParams = GetBuiltInParams();
+
+		for (const FPropertyBagPropertyDesc& Desc : BagStruct->GetPropertyDescs())
+		{
+			// Try to extract as float (covers float and double bag types).
+			float FloatValue = 0.0f;
+			if (!TryGetFloat(Store, Desc, FloatValue))
+			{
+				// Non-float type (bool, int32, etc.) -forward to MetaSound only.
+				if (MSGen.IsValid())
+				{
+					ForwardToMetaSound(*MSGen, Store, Desc);
+				}
+				continue;
+			}
+
+			// If MetaSound has a matching input, send it there and skip the built-in handler.
+			if (MSGen.IsValid() && HasMetaSoundInput(Desc.Name))
+			{
+				MSGen->SetInputValue<float>(Desc.Name, FloatValue);
+				continue;
+			}
+
+			// Check if this is a built-in parameter with a dedicated DSP handler.
+			if (const FBuiltInParamHandler* Handler = BuiltInParams.Find(Desc.Name))
+			{
+				(this->**Handler)(FloatValue);
+				continue;
+			}
+
+			// Unknown float param with no MetaSound input -ignore.
+		}
+	}
+
+	void FSubsonicGenerator::SetVolume(float dB)
+	{
+		VolumeLinear = Audio::ConvertToLinear(dB);
+	}
+
+	void FSubsonicGenerator::SetPitchShift(float Semitones)
+	{
+		constexpr int32 PitchShiftInterpolationFrames = 64;
+		const float Ratio = FMath::Pow(2.0f, Semitones / 12.0f);
+		Resampler->SetFrameRatio(Ratio, PitchShiftInterpolationFrames);
+		if (!FMath::IsNearlyZero(Semitones))
+		{
+			bPitchShiftActive = true;
+		}
+	}
+
+	void FSubsonicGenerator::SetHighpassCutoff(float FreqHz)
+	{
+		if (FreqHz <= 0.0f)
+		{
+			HighpassFilter.SetEnabled(false);
+		}
+		else
+		{
+			HighpassFilter.SetEnabled(true);
+			HighpassFilter.SetFrequency(FreqHz);
+		}
+	}
+
+	void FSubsonicGenerator::SetLowpassCutoff(float FreqHz)
+	{
+		if (FreqHz <= 0.0f || FreqHz >= SampleRate * 0.5f)
+		{
+			LowpassFilter.SetEnabled(false);
+		}
+		else
+		{
+			LowpassFilter.SetEnabled(true);
+			LowpassFilter.SetFrequency(FreqHz);
+		}
+	}
+
+	void FSubsonicGenerator::SetFadeOutTime(float Seconds)
+	{
+		if (FadeGainDelta == 0.0f)
+		{
+			FadeOutTimeSec = Seconds;
+		}
+	}
+
+	void FSubsonicGenerator::SetSource(Audio::IAudioMixerGeneratorSource* InSource)
+	{
+		Source = InSource;
+	}
+
+	void FSubsonicGenerator::Play()
+	{
+		if (Source)
+		{
+			Source->Play();
+		}
+	}
+
+	void FSubsonicGenerator::Stop()
+	{
+		if (FadeGainDelta != 0.0f)
+		{
+			return;
+		}
+
+		if (FadeOutTimeSec > 0.0f)
+		{
+			// Begin fade -OnGenerateAudio will null Source when fade completes.
+			FadeGainDelta = 1.0f / (FadeOutTimeSec * SampleRate);
+			return;
+		}
+
+		// Immediate stop.
+		if (Source)
+		{
+			Source->Stop();
+			Source = nullptr;
+		}
+	}
+
+	Audio::IAudioMixerGeneratorSource* FSubsonicGenerator::GetSource() const
+	{
+		return Source;
+	}
+
+	int32 FSubsonicGenerator::OnGenerateAudio(float* OutAudio, int32 InNumSamples)
+	{
+		// Apply deferred params once graph becomes ready.
+		if (DeferredParams && bGraphReady.load(std::memory_order_acquire))
+		{
+			ApplyParameters(*DeferredParams);
+			DeferredParams.Reset();
+		}
+
+		const int32 NumChannels = InnerGenerator->GetNumChannels();
+		const int32 NumFrames = (NumChannels > 0) ? (InNumSamples / NumChannels) : InNumSamples;
+
+		int32 Result = 0;
+
+		// Signal chain: InnerGenerator ->[Resample/PitchShift] ->[Highpass] ->[Lowpass] ->Volume ->Output
+		if (bPitchShiftActive)
+		{
+			// Pitch shift active: generate more/fewer inner samples, then resample to output size.
+			const int32 NumInputFrames = Resampler->GetNumInputFramesNeededToProduceOutputFrames(NumFrames);
+			int32 NumInputSamples = NumInputFrames * NumChannels;
+
+			const int32 LeftoverSamples = ResampleInputBuffer.Num();
+			ResampleInputBuffer.SetNumUninitialized(FMath::Max(LeftoverSamples, NumInputSamples));
+			if (NumInputSamples > LeftoverSamples)
+			{
+				NumInputSamples = LeftoverSamples + InnerGenerator->OnGenerateAudio(ResampleInputBuffer.GetData() + LeftoverSamples, NumInputSamples - LeftoverSamples);
+			}
+			else
+			{
+				NumInputSamples = LeftoverSamples;
+			}
+
+			int32 InputFramesConsumed = 0;
+			int32 OutputFramesProduced = 0;
+			Resampler->ProcessInterleaved(
+				TArrayView<const float>(ResampleInputBuffer.GetData(), NumInputSamples),
+				TArrayView<float>(OutAudio, InNumSamples),
+				InputFramesConsumed,
+				OutputFramesProduced);
+
+			const int32 ProducedSamples = OutputFramesProduced * NumChannels;
+			Result = ProducedSamples;
+
+			// Save leftover frames for next run.
+			const int32 InputSamplesConsumed = InputFramesConsumed * NumChannels;
+			const int32 NewLeftoverSamples = NumInputSamples - InputSamplesConsumed;
+			if (NewLeftoverSamples > 0)
+			{
+				FMemory::Memmove(ResampleInputBuffer.GetData(), ResampleInputBuffer.GetData() + InputSamplesConsumed, NewLeftoverSamples * sizeof(float));
+			}
+			ResampleInputBuffer.SetNum(NewLeftoverSamples);
+		}
+		else
+		{
+			Result = InnerGenerator->OnGenerateAudio(OutAudio, InNumSamples);
+		}
+
+		// Apply filters in-place (no-op when disabled).
+		HighpassFilter.ProcessAudio(OutAudio, InNumSamples, OutAudio);
+		LowpassFilter.ProcessAudio(OutAudio, InNumSamples, OutAudio);
+
+		// Apply volume and fade-out in a single pass via SIMD-optimized ramp.
+		{
+			const float FadeStart = FadeGain;
+			if (FadeGainDelta != 0.0f)
+			{
+				FadeGain = FMath::Max(0.0f, FadeGain - FadeGainDelta * NumFrames);
+			}
+
+			const float StartGain = VolumeLinear * FadeStart;
+			const float EndGain = VolumeLinear * FadeGain;
+			if (StartGain != 1.0f || EndGain != 1.0f)
+			{
+				Audio::ArrayFade(TArrayView<float>(OutAudio, InNumSamples), StartGain, EndGain);
+			}
+		}
+
+		// When the generator is done (fade finished or inner generator ended),
+		// stop the source and null the pointer so the subscriber knows to clean up.
+		if (Source && IsFinished())
+		{
+			Source->Stop();
+			Source = nullptr;
+		}
+
+		return Result;
+	}
+
+	int32 FSubsonicGenerator::GetNumChannels() const
+	{
+		return InnerGenerator->GetNumChannels();
+	}
+
+	bool FSubsonicGenerator::IsFinished() const
+	{
+		return FadeGain <= 0.0f || InnerGenerator->IsFinished();
+	}
+
+} // namespace UE::Subsonic
diff --git a/Engine/Plugins/Experimental/Subsonic/Source/Subsonic/Private/CoreEventSubscribers/SubsonicGenerator.h b/Engine/Plugins/Experimental/Subsonic/Source/Subsonic/Private/CoreEventSubscribers/SubsonicGenerator.h
new file mode 100644
--- /dev/null
+++ b/Engine/Plugins/Experimental/Subsonic/Source/Subsonic/Private/CoreEventSubscribers/SubsonicGenerator.h
@@ -0,0 +1,126 @@
+// Copyright Epic Games, Inc. All Rights Reserved.
+
+#pragma once
+
+#include "DSP/Filter.h"
+#include "DSP/RuntimeResampler.h"
+#include "MetasoundGenerator.h"
+#include "Sound/SoundGenerator.h"
+
+#include <atomic>
+
+namespace Audio { class IAudioMixerGeneratorSource; }
+
+namespace UE::Subsonic
+{
+	class FWaveGenerator;
+	struct FSubsonicParameterStore;
+
+	// Wrapper ISoundGenerator that owns either a WaveGenerator or MetsoundGenerator inner generator.
+	// Applies DSP (volume, pitch shift, highpass, lowpass) after delegating audio
+	//  generation to the inner generator.
+	//
+	// After initialization, these generators are owned by the audio render thread and should be communicated
+	//  with only through SubsonicRelay.
+	class FSubsonicGenerator : public ISoundGenerator
+	{
+	public:
+		// Wave path: inner generator is an FWaveGenerator (pure wave player).
+		explicit FSubsonicGenerator(TSharedRef<FWaveGenerator> InWave, float InSampleRate);
+
+		// MetaSound path: inner generator is an FMetasoundGenerator (returned by UMetaSoundSource).
+		explicit FSubsonicGenerator(ISoundGeneratorPtr InMetaSoundGenerator, float InSampleRate);
+
+		virtual ~FSubsonicGenerator() override;
+
+		//~ Begin ISoundGenerator
+		virtual int32 OnGenerateAudio(float* OutAudio, int32 InNumSamples) override;
+		virtual int32 GetNumChannels() const override;
+		virtual bool IsFinished() const override;
+		//~ End ISoundGenerator
+
+		// Audio render thread: apply named parameters from the store. Built-in parameters are
+		// handled by dedicated DSP handlers, unless the MetasoundGenerator provides a matching input.
+		void ApplyParameters(const FSubsonicParameterStore& Store);
+
+		// Associate a GeneratorSource with this generator, so the generator can be in charge of source state
+		//  (play, stop, pause, etc.)
+		void SetSource(Audio::IAudioMixerGeneratorSource* InSource);
+
+		// Start playback via the associated GeneratorSource.
+		void Play();
+
+		// Stop playback. If FadeOutTime > 0, begins a fade; otherwise calls Source->Stop() immediately.
+		void Stop();
+
+		// Game thread: returns the bare source pointer. Null means the generator is done
+		//  and the subscriber can safely destroy the source.
+		Audio::IAudioMixerGeneratorSource* GetSource() const;
+
+	private:
+		using FBuiltInParamHandler = void (FSubsonicGenerator::*)(float);
+		static const TMap<FName, FBuiltInParamHandler>& GetBuiltInParams();
+
+		// Attempts to extract a float from a property bag descriptor (handles float and double types).
+		static bool TryGetFloat(const FSubsonicParameterStore& Store, const FPropertyBagPropertyDesc& Desc, float& OutValue);
+
+		// Returns true if the pinned MetaSound generator has an input with the given name.
+		bool HasMetaSoundInput(FName Name) const;
+
+		// Forwards a non-float parameter to the pinned MetaSound generator.
+		void ForwardToMetaSound(Metasound::FMetasoundGenerator& MSGen, const FSubsonicParameterStore& Store, const FPropertyBagPropertyDesc& Desc);
+
+		// Built-in parameter handlers.
+		void SetVolume(float dB);
+		void SetPitchShift(float Semitones);
+		void SetHighpassCutoff(float FreqHz);
+		void SetLowpassCutoff(float FreqHz);
+		void SetFadeOutTime(float Seconds);
+
+		// Initializes DSP state (filters, resampler) after InnerGenerator is set.
+		void InitDSP();
+
+		ISoundGeneratorPtr InnerGenerator;
+
+		// Non-null for MetaSound path only.
+		TWeakPtr<Metasound::FMetasoundGenerator, ESPMode::ThreadSafe> MetaSoundGeneratorWeak;
+
+		// True once the MetaSound graph has been built and its inputs are available.
+		// Set via AddGraphSetCallback; always true for wave generators.
+		std::atomic<bool> bGraphReady{ false };
+
+		// Handle for the graph-set callback so we can unregister in the destructor.
+		FDelegateHandle GraphSetCallbackHandle;
+
+		// Holds the most recent SetParameters payload when the MetaSound graph is not yet ready.
+		// Applied and cleared once bGraphReady becomes true. Last-writer-wins.
+		TSharedPtr<FSubsonicParameterStore> DeferredParams;
+
+		// Current volume as linear gain. Updated from the "Volume" parameter (interpreted as dB).
+		// 0 dB = 1.0 linear (unity gain). Render thread only.
+		float VolumeLinear = 1.0f;
+
+		// Cached sample rate from the audio device.
+		float SampleRate = 0.0f;
+
+		// DSP processors applied in the signal chain after the inner generator.
+		Audio::FBiquadFilter HighpassFilter;
+		Audio::FBiquadFilter LowpassFilter;
+
+		// Pitch shifting via resampling.
+		TUniquePtr<Audio::FRuntimeResampler> Resampler;
+		TArray<float> ResampleInputBuffer;
+		bool bPitchShiftActive = false;
+
+		// Fade-out envelope state. Render thread only.
+		// FadeGainDelta > 0 means a fade is in progress; FadeGain <= 0 means it has completed.
+		float FadeOutTimeSec = 0.0f;
+		float FadeGain = 1.0f;
+		float FadeGainDelta = 0.0f;
+
+		// Non-owning pointer to the associated GeneratorSource. Set on the game thread
+		//  at creation, nulled on the render thread when the generator is done.
+		Audio::IAudioMixerGeneratorSource* Source = nullptr;
+	};
+
+} // namespace UE::Subsonic
diff --git a/Engine/Plugins/Experimental/Subsonic/Source/Subsonic/Private/CoreEventSubscribers/SubsonicGeneratorSourceSubscriber.cpp b/Engine/Plugins/Experimental/Subsonic/Source/Subsonic/Private/CoreEventSubscribers/SubsonicGeneratorSourceSubscriber.cpp
--- a/Engine/Plugins/Experimental/Subsonic/Source/Subsonic/Private/CoreEventSubscribers/SubsonicGeneratorSourceSubscriber.cpp
+++ b/Engine/Plugins/Experimental/Subsonic/Source/Subsonic/Private/CoreEventSubscribers/SubsonicGeneratorSourceSubscriber.cpp
@@ -3,11 +3,13 @@
 #include "CoreEventSubscribers/SubsonicGeneratorSourceSubscriber.h"
 
 #include "Algo/Transform.h"
-#include "AudioDeviceManager.h"
 #include "AudioDevice.h"
+#include "AudioDeviceManager.h"
 #include "AudioMixerDevice.h"
 #include "Engine/Engine.h"
 #include "IAudioMixerGeneratorSource.h"
+#include "MetasoundSource.h"
+#include "SubsonicGenerator.h"
 #include "SubsonicLog.h"
 #include "SubsonicSubsystem.h"
 #include "SubsonicWaveGenerator.h"
@@ -16,8 +18,64 @@
 
 namespace UE::Subsonic
 {
+	void FInstanceSources::ApplyParameters(FSubsonicRelay& Relay)
+	{
+		if (!Generator)
+		{
+			return;
+		}
+
+		// Always allocate a fresh merged store - the audio thread owns its TSharedPtr.
+		TSharedPtr<FSubsonicParameterStore> Merged = MakeShared<FSubsonicParameterStore>();
+		Merged->MergeFrom(AuthoredParams);
+		Merged->MergeFrom(RuntimeParams);
+
+		// Enqueue for the audio thread via the relay's global command queue.
+		FRelayCommand Cmd;
+		Cmd.Type = FRelayCommand::EType::SetParameters;
+		Cmd.Params = MoveTemp(Merged);
+		Cmd.Target = Generator;
+		Relay.EnqueueCommand(MoveTemp(Cmd));
+	}
+
 	namespace SubscriberPrivate
 	{
+		// Creates an FSubsonicGenerator from a USoundWave, using the MetaSound path
+		// if the sound is a UMetaSoundSource, otherwise the wave path.
+		TSharedPtr<FSubsonicGenerator> CreateGenerator(USoundWave& Sound, Audio::FMixerDevice& AudioDevice)
+		{
+			const float SampleRate = AudioDevice.GetSampleRate();
+
+			if (UMetaSoundSource* MetaSound = Cast<UMetaSoundSource>(&Sound))
+			{
+				// Ensure the MetaSound graph is registered and compiled.
+				MetaSound->InitResources();
+
+				FSoundGeneratorInitParams InitParams;
+				InitParams.AudioDeviceID = AudioDevice.DeviceID;
+				InitParams.SampleRate = SampleRate;
+				InitParams.AudioMixerNumOutputFrames = AudioDevice.GetNumOutputFrames();
+				InitParams.NumChannels = MetaSound->NumChannels;
+				InitParams.NumFramesPerCallback = AudioDevice.GetNumOutputFrames();
+				InitParams.AudioComponentId = 0;
+				InitParams.InstanceID = 0;
+				InitParams.bIsPreviewSound = false;
+
+				TArray<FAudioParameter> DefaultParams;
+				MetaSound->InitParameters(DefaultParams, NAME_None);
+				ISoundGeneratorPtr MSGenPtr = MetaSound->CreateSoundGenerator(InitParams, MoveTemp(DefaultParams));
+				if (!MSGenPtr.IsValid())
+				{
+					return nullptr;
+				}
+
+				return MakeShared<FSubsonicGenerator>(MSGenPtr, SampleRate);
+			}
+
+			TSharedRef<FWaveGenerator> WaveGen = MakeShared<FWaveGenerator>(Sound.CreateSoundWaveProxy(), &AudioDevice);
+			return MakeShared<FSubsonicGenerator>(WaveGen, SampleRate);
+		}
+
 		template <typename HandleType>
 		Audio::IAudioMixerGeneratorSource* PlaySound(
 			USubsonicGeneratorSourceSubscriber& Subscriber,
@@ -25,6 +83,8 @@ namespace UE::Subsonic
 			const HandleType& InHandle,
 			TMap<HandleType, FGeneratorSourceCache>& OutSourceCache,
 			FName Name,
+			const FSubsonicParameterStore& Params,
+			FSubsonicRelay& Relay,
 			bool bReleaseExisting)
 		{
 			Audio::IAudioMixerGeneratorSource* ReturnSource = nullptr;
@@ -33,20 +93,34 @@ namespace UE::Subsonic
 			{
 				if (Audio::FMixerDevice* AudioDevice = Subscriber.GetMutableMixerDevice())
 				{
+					TSharedPtr<FSubsonicGenerator> Generator = CreateGenerator(Sound, *AudioDevice);
+					if (!Generator.IsValid())
+					{
+						return nullptr;
+					}
+
 					FGeneratorSourceCache& Cache = OutSourceCache.FindOrAdd(InHandle);
 					FInstanceSources& InstanceSources = Cache.FindOrAdd(Name);
 
-					TSharedRef<FWaveGenerator> Generator = MakeShared<FWaveGenerator>(Sound.CreateSoundWaveProxy(), AudioDevice);
-					TUniquePtr<Audio::IAudioMixerGeneratorSource> NewSource = AudioDevice->CreateGeneratorSource(Generator);
+					TUniquePtr<Audio::IAudioMixerGeneratorSource> NewSource = AudioDevice->CreateGeneratorSource(Generator.ToSharedRef());
 					if (NewSource.IsValid())
 					{
 						ReturnSource = NewSource.Get();
 						if (InstanceSources.ActiveSource != nullptr)
 						{
-							InstanceSources.ReleasedSources.Add(MoveTemp(InstanceSources.ActiveSource));
+							FReleasedSource Released;
+							Released.Source = MoveTemp(InstanceSources.ActiveSource);
+							Released.Generator = InstanceSources.Generator;
+							InstanceSources.ReleasedSources.Add(MoveTemp(Released));
 						}
 
-						NewSource->Play();
+						InstanceSources.Generator = Generator;
+						InstanceSources.AuthoredParams = Params;
+						InstanceSources.RuntimeParams.Reset();
+						InstanceSources.ApplyParameters(Relay);
+
+						Generator->SetSource(NewSource.Get());
+						Relay.EnqueueCommand(FRelayCommand{ FRelayCommand::EType::Play, nullptr, Generator });
 						InstanceSources.ActiveSource = MoveTemp(NewSource);
 					}
 				}
@@ -67,13 +141,14 @@ namespace UE::Subsonic
 				{
 					FInstanceSources& Instances = SourcesItem.Value;
 
-					Instances.ReleasedSources.RemoveAll([](TUniquePtr<Audio::IAudioMixerGeneratorSource>& Source)
+					Instances.ReleasedSources.RemoveAll([](const FReleasedSource& Released)
 						{
-							return Source->IsDone();
+							return !Released.Generator || Released.Generator->GetSource() == nullptr;
 						});
-					if (Instances.ActiveSource != nullptr && Instances.ActiveSource->IsDone())
+					if (Instances.ActiveSource != nullptr && Instances.Generator && Instances.Generator->GetSource() == nullptr)
 					{
 						Instances.ActiveSource.Reset();
+						Instances.Generator.Reset();
 					}
 					if (Instances.ActiveSource == nullptr && Instances.ReleasedSources.IsEmpty())
 					{
@@ -99,6 +174,7 @@ namespace UE::Subsonic
 
 		template <typename HandleType>
 		void StopSound(
+			FSubsonicRelay& Relay,
 			const HandleType& InHandle,
 			TMap<HandleType, FGeneratorSourceCache>& OutSourceCache,
 			FName Name)
@@ -106,13 +182,65 @@ namespace UE::Subsonic
 			FGeneratorSourceCache* Cache = OutSourceCache.Find(InHandle);
 			if (Cache != nullptr)
 			{
-				Cache->Remove(Name);
+				if (FInstanceSources* Instances = Cache->Find(Name))
+				{
+					if (Instances->ActiveSource != nullptr)
+					{
+						FReleasedSource Released;
+						Released.Source = MoveTemp(Instances->ActiveSource);
+						Released.Generator = Instances->Generator;
+						Instances->ReleasedSources.Add(MoveTemp(Released));
+
+						Relay.EnqueueCommand(FRelayCommand{
+							FRelayCommand::EType::Stop, nullptr, Instances->Generator
+						});
+
+						Instances->Generator.Reset();
+					}
+				}
+			}
+		}
+
+		template <typename HandleType>
+		void SetParameters(
+			FSubsonicRelay& Relay,
+			const HandleType& InHandle,
+			TMap<HandleType, FGeneratorSourceCache>& SourceCache,
+			FName SourceName,
+			const FSubsonicParameterStore& Params)
+		{
+			if (FGeneratorSourceCache* Cache = SourceCache.Find(InHandle))
+			{
+				if (FInstanceSources* Instances = Cache->Find(SourceName))
+				{
+					Instances->RuntimeParams.MergeFrom(Params);
+					Instances->ApplyParameters(Relay);
+				}
+			}
+		}
+
+		template <typename HandleType>
+		void SetAuthoredParameters(
+			FSubsonicRelay& Relay,
+			const HandleType& InHandle,
+			TMap<HandleType, FGeneratorSourceCache>& SourceCache,
+			FName SourceName,
+			const FSubsonicParameterStore& Params)
+		{
+			if (FGeneratorSourceCache* Cache = SourceCache.Find(InHandle))
+			{
+				if (FInstanceSources* Instances = Cache->Find(SourceName))
+				{
+					Instances->AuthoredParams = Params;
+					Instances->ApplyParameters(Relay);
+				}
 			}
 		}
 	} // namespace SubscriberPrivate
 
 	USubsonicGeneratorSourceSubscriber::USubsonicGeneratorSourceSubscriber()
 		: FTickableGameObject(ETickableTickType::Never)
+		, Relay(MakeShared<FSubsonicRelay>())
 	{
 	}
 
@@ -143,50 +271,97 @@ namespace UE::Subsonic
 		CollectionSources.Remove(InHandle);
 	}
 
-	void USubsonicGeneratorSourceSubscriber::PlaySound(FName Name, USoundWave& Sound, bool bReleaseExisting)
+	void USubsonicGeneratorSourceSubscriber::PlaySound(FName Name, USoundWave& Sound, const FSubsonicParameterStore& Params, bool bReleaseExisting)
 	{
 		const int32 FakeHandle = 0;
-		SubscriberPrivate::PlaySound(*this, Sound, FakeHandle, GlobalSources, Name, bReleaseExisting);
+		SubscriberPrivate::PlaySound(*this, Sound, FakeHandle, GlobalSources, Name, Params, *Relay, bReleaseExisting);
 	}
 
-	void USubsonicGeneratorSourceSubscriber::PlaySound(const FCollectionHandle& InHandle, FName Name, USoundWave& Sound, bool bReleaseExisting)
+	void USubsonicGeneratorSourceSubscriber::PlaySound(const FCollectionHandle& InHandle, FName Name, USoundWave& Sound, const FSubsonicParameterStore& Params, bool bReleaseExisting)
 	{
-		SubscriberPrivate::PlaySound(*this, Sound, InHandle, CollectionSources, Name, bReleaseExisting);
+		SubscriberPrivate::PlaySound(*this, Sound, InHandle, CollectionSources, Name, Params, *Relay, bReleaseExisting);
 	}
 
-	void USubsonicGeneratorSourceSubscriber::PlaySound(const FEventHandle& InHandle, FName Name, USoundWave& Sound, bool bReleaseExisting)
+	void USubsonicGeneratorSourceSubscriber::PlaySound(const FEventHandle& InHandle, FName Name, USoundWave& Sound, const FSubsonicParameterStore& Params, bool bReleaseExisting)
 	{
-		SubscriberPrivate::PlaySound(*this, Sound, InHandle, EventSources, Name, bReleaseExisting);
+		SubscriberPrivate::PlaySound(*this, Sound, InHandle, EventSources, Name, Params, *Relay, bReleaseExisting);
 	}
 
-	void USubsonicGeneratorSourceSubscriber::PlaySound(const FActionHandle& InHandle, FName Name, USoundWave& Sound, bool bReleaseExisting)
+	void USubsonicGeneratorSourceSubscriber::PlaySound(const FActionHandle& InHandle, FName Name, USoundWave& Sound, const FSubsonicParameterStore& Params, bool bReleaseExisting)
 	{
-		SubscriberPrivate::PlaySound(*this, Sound, InHandle, ActionSources, Name, bReleaseExisting);
+		SubscriberPrivate::PlaySound(*this, Sound, InHandle, ActionSources, Name, Params, *Relay, bReleaseExisting);
 	}
 
 	void USubsonicGeneratorSourceSubscriber::StopSound(FName Name)
 	{
 		const int32 FakeHandle = 0;
-		SubscriberPrivate::StopSound(FakeHandle, GlobalSources, Name);
+		SubscriberPrivate::StopSound(*Relay, FakeHandle, GlobalSources, Name);
 	}
 
 	void USubsonicGeneratorSourceSubscriber::StopSound(const FCollectionHandle& InHandle, FName Name)
 	{
-		SubscriberPrivate::StopSound(InHandle, CollectionSources, Name);
+		SubscriberPrivate::StopSound(*Relay, InHandle, CollectionSources, Name);
 	}
 
 	void USubsonicGeneratorSourceSubscriber::StopSound(const FEventHandle& InHandle, FName Name)
 	{
-		SubscriberPrivate::StopSound(InHandle, EventSources, Name);
+		SubscriberPrivate::StopSound(*Relay, InHandle, EventSources, Name);
 	}
 
 	void USubsonicGeneratorSourceSubscriber::StopSound(const FActionHandle& InHandle, FName Name)
 	{
-		SubscriberPrivate::StopSound(InHandle, ActionSources, Name);
+		SubscriberPrivate::StopSound(*Relay, InHandle, ActionSources, Name);
+	}
+
+	void USubsonicGeneratorSourceSubscriber::SetParameters(FName SourceName, const FSubsonicParameterStore& Params)
+	{
+		const int32 FakeHandle = 0;
+		SubscriberPrivate::SetParameters(*Relay, FakeHandle, GlobalSources, SourceName, Params);
+	}
+
+	void USubsonicGeneratorSourceSubscriber::SetParameters(const FCollectionHandle& InHandle, FName SourceName, const FSubsonicParameterStore& Params)
+	{
+		SubscriberPrivate::SetParameters(*Relay, InHandle, CollectionSources, SourceName, Params);
+	}
+
+	void USubsonicGeneratorSourceSubscriber::SetParameters(const FEventHandle& InHandle, FName SourceName, const FSubsonicParameterStore& Params)
+	{
+		SubscriberPrivate::SetParameters(*Relay, InHandle, EventSources, SourceName, Params);
+	}
+
+	void USubsonicGeneratorSourceSubscriber::SetParameters(const FActionHandle& InHandle, FName SourceName, const FSubsonicParameterStore& Params)
+	{
+		SubscriberPrivate::SetParameters(*Relay, InHandle, ActionSources, SourceName, Params);
+	}
+
+	void USubsonicGeneratorSourceSubscriber::SetAuthoredParameters(FName SourceName, const FSubsonicParameterStore& Params)
+	{
+		const int32 FakeHandle = 0;
+		SubscriberPrivate::SetAuthoredParameters(*Relay, FakeHandle, GlobalSources, SourceName, Params);
+	}
+
+	void USubsonicGeneratorSourceSubscriber::SetAuthoredParameters(const FCollectionHandle& InHandle, FName SourceName, const FSubsonicParameterStore& Params)
+	{
+		SubscriberPrivate::SetAuthoredParameters(*Relay, InHandle, CollectionSources, SourceName, Params);
+	}
+
+	void USubsonicGeneratorSourceSubscriber::SetAuthoredParameters(const FEventHandle& InHandle, FName SourceName, const FSubsonicParameterStore& Params)
+	{
+		SubscriberPrivate::SetAuthoredParameters(*Relay, InHandle, EventSources, SourceName, Params);
+	}
+
+	void USubsonicGeneratorSourceSubscriber::SetAuthoredParameters(const FActionHandle& InHandle, FName SourceName, const FSubsonicParameterStore& Params)
+	{
+		SubscriberPrivate::SetAuthoredParameters(*Relay, InHandle, ActionSources, SourceName, Params);
 	}
 
 	void USubsonicGeneratorSourceSubscriber::Tick(float DeltaTime)
 	{
+		if (Audio::FMixerDevice* MixerDevice = GetMutableMixerDevice())
+		{
+			Relay->Tick(*MixerDevice);
+		}
+
 		SubscriberPrivate::ClearStaleSounds(GlobalSources);
 		SubscriberPrivate::ClearStaleSounds(CollectionSources);
 		SubscriberPrivate::ClearStaleSounds(EventSources);
diff --git a/Engine/Plugins/Experimental/Subsonic/Source/Subsonic/Private/CoreEventSubscribers/SubsonicGeneratorSourceSubscriber.h b/Engine/Plugins/Experimental/Subsonic/Source/Subsonic/Private/CoreEventSubscribers/SubsonicGeneratorSourceSubscriber.h
--- a/Engine/Plugins/Experimental/Subsonic/Source/Subsonic/Private/CoreEventSubscribers/SubsonicGeneratorSourceSubscriber.h
+++ b/Engine/Plugins/Experimental/Subsonic/Source/Subsonic/Private/CoreEventSubscribers/SubsonicGeneratorSourceSubscriber.h
@@ -2,9 +2,11 @@
 #pragma once
 
 #include "CoreEventSubscribers/SubsonicEventSubscriberBase.h"
+#include "CoreEventSubscribers/SubsonicRelay.h"
 #include "IAudioMixerGeneratorSource.h"
 #include "SubsonicEventSystemInterface.h"
 #include "SubsonicHandles.h"
+#include "SubsonicParameterStore.h"
 #include "Subsystems/AudioEngineSubsystem.h"
 #include "Tickable.h"
 
@@ -15,13 +17,34 @@ class USoundWave;
 
 namespace UE::Subsonic
 {
+	class FSubsonicGenerator;
+
+	// A source that is no longer addressable from actions but still playing.
+	// Paired with its generator so ClearStaleSounds can check completion.
+	struct FReleasedSource
+	{
+		TUniquePtr<Audio::IAudioMixerGeneratorSource> Source;
+		TSharedPtr<FSubsonicGenerator> Generator;
+	};
+
 	struct FInstanceSources
 	{
 		// Currently addressable generator source via handle
 		TUniquePtr<Audio::IAudioMixerGeneratorSource> ActiveSource;
 
 		// Sources no longer addressable from actions but actively playing.
-		TArray<TUniquePtr<Audio::IAudioMixerGeneratorSource>> ReleasedSources;
+		TArray<FReleasedSource> ReleasedSources;
+
+		// The active generator, shared with the relay command queue so that
+		// SetParameters commands can target it directly on the render thread.
+		TSharedPtr<FSubsonicGenerator> Generator;
+
+		FSubsonicParameterStore AuthoredParams;   // from action definition (lowest priority)
+		FSubsonicParameterStore RuntimeParams;    // from post-playback SetParameters (highest priority)
+
+		// Merges all three layers (authored < trigger < runtime) and enqueues a
+		// SetParameters command into the relay's global command queue.
+		void ApplyParameters(FSubsonicRelay& Relay);
 	};
 
 	using FGeneratorSourceCache = TMap<FName, FInstanceSources>;
@@ -30,7 +53,7 @@ namespace UE::Subsonic
 	 * Interface to GeneratorSource playback Subsonic implementation.
 	 */
 	UCLASS(MinimalAPI, BlueprintType)
-		class USubsonicGeneratorSourceSubscriber final : public USubsonicEventSubscriberBase, public FTickableGameObject
+	class USubsonicGeneratorSourceSubscriber final : public USubsonicEventSubscriberBase, public FTickableGameObject
 	{
 		GENERATED_BODY()
 
@@ -45,20 +68,34 @@ namespace UE::Subsonic
 		virtual void OnCollectionRegistered(const FCollectionHandle& InEvent) override;
 		virtual void OnCollectionUnregistered(const FCollectionHandle& InEvent) override;
 
-		void PlaySound(FName Name, USoundWave& Sound, bool bReleaseExisting = true);
-		void PlaySound(const FCollectionHandle& InHandle, FName Name, USoundWave& Sound, bool bReleaseExisting = true);
-		void PlaySound(const FEventHandle& InHandle, FName Name, USoundWave& Sound, bool bReleaseExisting = true);
-		void PlaySound(const FActionHandle& InHandle, FName Name, USoundWave& Sound, bool bReleaseExisting = true);
+		void PlaySound(FName Name, USoundWave& Sound, const FSubsonicParameterStore& Params, bool bReleaseExisting = true);
+		void PlaySound(const FCollectionHandle& InHandle, FName Name, USoundWave& Sound, const FSubsonicParameterStore& Params, bool bReleaseExisting = true);
+		void PlaySound(const FEventHandle& InHandle, FName Name, USoundWave& Sound, const FSubsonicParameterStore& Params, bool bReleaseExisting = true);
+		void PlaySound(const FActionHandle& InHandle, FName Name, USoundWave& Sound, const FSubsonicParameterStore& Params, bool bReleaseExisting = true);
 
 		void StopSound(FName Name);
 		void StopSound(const FCollectionHandle& InHandle, FName Name);
 		void StopSound(const FEventHandle& InHandle, FName Name);
 		void StopSound(const FActionHandle& InHandle, FName Name);
 
+		void SetParameters(FName SourceName, const FSubsonicParameterStore& Params);
+		void SetParameters(const FCollectionHandle& InHandle, FName SourceName, const FSubsonicParameterStore& Params);
+		void SetParameters(const FEventHandle& InHandle, FName SourceName, const FSubsonicParameterStore& Params);
+		void SetParameters(const FActionHandle& InHandle, FName SourceName, const FSubsonicParameterStore& Params);
+
+		SUBSONIC_API void SetAuthoredParameters(FName SourceName, const FSubsonicParameterStore& Params);
+		SUBSONIC_API void SetAuthoredParameters(const FCollectionHandle& InHandle, FName SourceName, const FSubsonicParameterStore& Params);
+		SUBSONIC_API void SetAuthoredParameters(const FEventHandle& InHandle, FName SourceName, const FSubsonicParameterStore& Params);
+		SUBSONIC_API void SetAuthoredParameters(const FActionHandle& InHandle, FName SourceName, const FSubsonicParameterStore& Params);
+
 		// FTickableGameObject interface
 		virtual void Tick(float DeltaTime) override;
+		virtual bool IsTickableInEditor() const override { return true; }
 		virtual TStatId GetStatId() const override;
+
 	private:
+		TSharedRef<FSubsonicRelay> Relay;
+
 		TMap<int32, FGeneratorSourceCache> GlobalSources;
 		TMap<FCollectionHandle, FGeneratorSourceCache> CollectionSources;
 		TMap<FEventHandle, FGeneratorSourceCache> EventSources;
diff --git a/Engine/Plugins/Experimental/Subsonic/Source/Subsonic/Private/CoreEventSubscribers/SubsonicRelay.cpp b/Engine/Plugins/Experimental/Subsonic/Source/Subsonic/Private/CoreEventSubscribers/SubsonicRelay.cpp
new file mode 100644
--- /dev/null
+++ b/Engine/Plugins/Experimental/Subsonic/Source/Subsonic/Private/CoreEventSubscribers/SubsonicRelay.cpp
@@ -0,0 +1,57 @@
+// Copyright Epic Games, Inc. All Rights Reserved.
+
+#include "CoreEventSubscribers/SubsonicRelay.h"
+
+#include "AudioMixerDevice.h"
+#include "SubsonicGenerator.h"
+
+namespace UE::Subsonic
+{
+	void FSubsonicRelay::EnqueueCommand(FRelayCommand&& Cmd)
+	{
+		FScopeLock ScopeLock(&CommandLock);
+		PendingCommands.Add(MoveTemp(Cmd));
+	}
+
+	void FSubsonicRelay::Tick(Audio::FMixerDevice& MixerDevice)
+	{
+		TArray<FRelayCommand> Commands;
+		{
+			FScopeLock ScopeLock(&CommandLock);
+			Commands = MoveTemp(PendingCommands);
+			// PendingCommands is now empty after move
+		}
+
+		if (Commands.IsEmpty())
+		{
+			return;
+		}
+
+		MixerDevice.AudioRenderThreadCommand([Cmds = MoveTemp(Commands)]()
+		{
+			for (const FRelayCommand& Cmd : Cmds)
+			{
+				if (!Cmd.Target)
+				{
+					continue;
+				}
+
+				switch (Cmd.Type)
+				{
+				case FRelayCommand::EType::SetParameters:
+					if (Cmd.Params)
+					{
+						Cmd.Target->ApplyParameters(*Cmd.Params);
+					}
+					break;
+				case FRelayCommand::EType::Play:
+					Cmd.Target->Play();
+					break;
+				case FRelayCommand::EType::Stop:
+					Cmd.Target->Stop();
+					break;
+				}
+			}
+		});
+	}
+} // namespace UE::Subsonic
diff --git a/Engine/Plugins/Experimental/Subsonic/Source/Subsonic/Private/CoreEventSubscribers/SubsonicRelay.h b/Engine/Plugins/Experimental/Subsonic/Source/Subsonic/Private/CoreEventSubscribers/SubsonicRelay.h
new file mode 100644
--- /dev/null
+++ b/Engine/Plugins/Experimental/Subsonic/Source/Subsonic/Private/CoreEventSubscribers/SubsonicRelay.h
@@ -0,0 +1,42 @@
+// Copyright Epic Games, Inc. All Rights Reserved.
+
+#pragma once
+
+#include "Containers/Array.h"
+#include "HAL/CriticalSection.h"
+#include "SubsonicParameterStore.h"
+#include "Templates/SharedPointer.h"
+
+namespace Audio { class FMixerDevice; }
+
+namespace UE::Subsonic
+{
+	class FSubsonicGenerator;
+
+	// Typed command for the per-device command queue on FSubsonicRelay. Commands are expected to be added on the game thread but can
+	//  be safely added from other threads. All commands in one game tick are pushed to the audio render thread in a single batch.
+	struct FRelayCommand
+	{
+		enum class EType : uint8 { SetParameters, Play, Stop };
+
+		EType Type = EType::SetParameters;
+		TSharedPtr<FSubsonicParameterStore> Params;      // SetParameters only
+		TSharedPtr<FSubsonicGenerator> Target;            // all types
+	};
+
+	// Per-device object responsible for communicating state changes to the audio render thread.
+	class FSubsonicRelay
+	{
+	public:
+		// Append a command to the pending queue (lock-protected).
+		void EnqueueCommand(FRelayCommand&& Cmd);
+
+		// Game thread (called from subscriber Tick): flush pending commands into a single
+		// AudioRenderThreadCommand lambda on the mixer device.
+		void Tick(Audio::FMixerDevice& MixerDevice);
+
+	private:
+		FCriticalSection CommandLock;
+		TArray<FRelayCommand> PendingCommands;
+	};
+} // namespace UE::Subsonic
diff --git a/Engine/Plugins/Experimental/Subsonic/Source/Subsonic/Private/CoreEventSubscribers/SubsonicWaveGenerator.h b/Engine/Plugins/Experimental/Subsonic/Source/Subsonic/Private/CoreEventSubscribers/SubsonicWaveGenerator.h
--- a/Engine/Plugins/Experimental/Subsonic/Source/Subsonic/Private/CoreEventSubscribers/SubsonicWaveGenerator.h
+++ b/Engine/Plugins/Experimental/Subsonic/Source/Subsonic/Private/CoreEventSubscribers/SubsonicWaveGenerator.h
@@ -15,20 +15,20 @@ class FSoundWaveProxyPlayer;
 namespace UE::Subsonic
 {
 
-	// A sound generator that can be used to play back a wave file
-
+	// A pure wave-file player. Volume scaling and relay communication are handled
+	// by the outer FSubsonicGenerator wrapper.
 	class FWaveGenerator : public ISoundGenerator
 	{
 	public:
-		explicit FWaveGenerator(const TSharedPtr<FSoundWaveProxy>& SoundWaveProxy, Audio::FMixerDevice* MixerDevice);
+		FWaveGenerator(const TSharedPtr<FSoundWaveProxy>& SoundWaveProxy, Audio::FMixerDevice* MixerDevice);
+
 		virtual ~FWaveGenerator() override;
 
 		virtual int32 OnGenerateAudio(float* OutAudio, int32 InNumSamples) override;
 		virtual int32 GetNumChannels() const override;
-
 		virtual bool IsFinished() const override;
-	
+
 	protected:
 		TUniquePtr<FSoundWaveProxyPlayer> Player;
 	};
-}
\ No newline at end of file
+}
diff --git a/Engine/Plugins/Experimental/Subsonic/Source/Subsonic/Private/SubsonicBuiltInParameters.cpp b/Engine/Plugins/Experimental/Subsonic/Source/Subsonic/Private/SubsonicBuiltInParameters.cpp
new file mode 100644
--- /dev/null
+++ b/Engine/Plugins/Experimental/Subsonic/Source/Subsonic/Private/SubsonicBuiltInParameters.cpp
@@ -0,0 +1,13 @@
+// Copyright Epic Games, Inc. All Rights Reserved.
+
+#include "SubsonicBuiltInParameters.h"
+
+namespace UE::Subsonic::BuiltInParameters
+{
+	const FName Volume         = "Volume";
+	const FName PitchShift     = "PitchShift";
+	const FName HighpassCutoff = "HighpassCutoff";
+	const FName LowpassCutoff  = "LowpassCutoff";
+	const FName FadeOutTime    = "FadeOutTime";
+
+} // namespace UE::Subsonic::BuiltInParameters
diff --git a/Engine/Plugins/Experimental/Subsonic/Source/Subsonic/Private/SubsonicExecutor.cpp b/Engine/Plugins/Experimental/Subsonic/Source/Subsonic/Private/SubsonicExecutor.cpp
--- a/Engine/Plugins/Experimental/Subsonic/Source/Subsonic/Private/SubsonicExecutor.cpp
+++ b/Engine/Plugins/Experimental/Subsonic/Source/Subsonic/Private/SubsonicExecutor.cpp
@@ -27,6 +27,7 @@ namespace UE::Subsonic
 		Unregister();
 
 		CollectionAccessor = MoveTemp(InOther.CollectionAccessor);
+		Parameters = MoveTemp(InOther.Parameters);
 
 		DeviceId = InOther.DeviceId;
 		InOther.DeviceId = ExecutorPrivate::InvalidDeviceId;
@@ -45,6 +46,7 @@ namespace UE::Subsonic
 		Unregister();
 
 		CollectionAccessor = MoveTemp(InOther.CollectionAccessor);
+		Parameters = MoveTemp(InOther.Parameters);
 
 		DeviceId = InOther.DeviceId;
 		InOther.DeviceId = ExecutorPrivate::InvalidDeviceId;
@@ -164,11 +166,22 @@ namespace UE::Subsonic
 		}
 
 		CollectionAccessor.Reset();
+		Parameters.Reset();
 		ExecutorId = ExecutorPrivate::InvalidId;
 	}
 
+	void FSubsonicExecutor::SetParameters(FSubsonicParameterStore InParams)
+	{
+		Parameters = MoveTemp(InParams);
+	}
+
+	const FSubsonicParameterStore& FSubsonicExecutor::GetParameters() const
+	{
+		return Parameters;
+	}
+
 	FString FSubsonicExecutor::ToString() const
 	{
 		return FString::Printf(TEXT("Id: %u, DeviceId: %u"), ExecutorId, DeviceId);
 	}
-} // namespace UE::Subsonic
\ No newline at end of file
+} // namespace UE::Subsonic
diff --git a/Engine/Plugins/Experimental/Subsonic/Source/Subsonic/Private/SubsonicParameterStore.cpp b/Engine/Plugins/Experimental/Subsonic/Source/Subsonic/Private/SubsonicParameterStore.cpp
new file mode 100644
--- /dev/null
+++ b/Engine/Plugins/Experimental/Subsonic/Source/Subsonic/Private/SubsonicParameterStore.cpp
@@ -0,0 +1,47 @@
+// Copyright Epic Games, Inc. All Rights Reserved.
+
+#include "SubsonicParameterStore.h"
+
+
+namespace UE::Subsonic
+{
+	bool FSubsonicParameterStore::HasParameter(FName Name) const
+	{
+		const UPropertyBag* BagStruct = Bag.GetPropertyBagStruct();
+		return BagStruct != nullptr && BagStruct->FindPropertyDescByName(Name) != nullptr;
+	}
+
+	void FSubsonicParameterStore::RemoveParameter(FName Name)
+	{
+		Bag.RemovePropertyByName(Name);
+	}
+
+	void FSubsonicParameterStore::Reset()
+	{
+		Bag = FInstancedPropertyBag{};
+	}
+
+	bool FSubsonicParameterStore::IsEmpty() const
+	{
+		return Bag.GetNumPropertiesInBag() == 0;
+	}
+
+	void FSubsonicParameterStore::MergeFrom(const FSubsonicParameterStore& Other)
+	{
+		const UPropertyBag* OtherBagStruct = Other.Bag.GetPropertyBagStruct();
+		if (!OtherBagStruct)
+		{
+			return;
+		}
+
+		const TConstArrayView<FPropertyBagPropertyDesc> OtherDescs = OtherBagStruct->GetPropertyDescs();
+		if (OtherDescs.IsEmpty())
+		{
+			return;
+		}
+
+		// Ensure all properties from Other exist in this bag, then overlay their values.
+		Bag.AddProperties(OtherDescs);
+		Bag.CopyMatchingValuesByName(Other.Bag);
+	}
+} // namespace UE::Subsonic
diff --git a/Engine/Plugins/Experimental/Subsonic/Source/Subsonic/Private/SubsonicSubsystem.cpp b/Engine/Plugins/Experimental/Subsonic/Source/Subsonic/Private/SubsonicSubsystem.cpp
--- a/Engine/Plugins/Experimental/Subsonic/Source/Subsonic/Private/SubsonicSubsystem.cpp
+++ b/Engine/Plugins/Experimental/Subsonic/Source/Subsonic/Private/SubsonicSubsystem.cpp
@@ -4,8 +4,8 @@
 
 #include "Algo/Transform.h"
 #include "AudioDefines.h"
-#include "Engine/GameInstance.h"
 #include "Engine/Engine.h"
+#include "Engine/GameInstance.h"
 #include "GameplayTagsManager.h"
 #include "Misc/CoreDefines.h"
 #include "Misc/Optional.h"
@@ -198,7 +198,7 @@ namespace UE::Subsonic
 		return nullptr;
 	}
 
-	void USubsonicSubsystem::ExecuteEventBP(const USubsonicEventCollectionExecutor* Executor, FGameplayTag EventTag)
+	void USubsonicSubsystem::ExecuteEventBP(USubsonicEventCollectionExecutor* Executor, FGameplayTag EventTag, FSubsonicParameterStore TriggerParams)
 	{
 		if (Executor)
 		{
@@ -206,7 +206,9 @@ namespace UE::Subsonic
 			const FSubsonicEventCollectionDefinition& Definition = EventCollection.GetDefinition();
 			if (const FSubsonicEvent* Event = Definition.GetEvents().Find(EventTag))
 			{
-				Executor->GetExecutor().ExecuteEvent(EventTag.GetTagName());
+				Executor->Executor.SetParameters(MoveTemp(TriggerParams));
+				Executor->Executor.ExecuteEvent(EventTag.GetTagName());
+				Executor->Executor.SetParameters(FSubsonicParameterStore{});
 			}
 		}
 	}
diff --git a/Engine/Plugins/Experimental/Subsonic/Source/Subsonic/Public/SubsonicBuiltInParameters.h b/Engine/Plugins/Experimental/Subsonic/Source/Subsonic/Public/SubsonicBuiltInParameters.h
new file mode 100644
--- /dev/null
+++ b/Engine/Plugins/Experimental/Subsonic/Source/Subsonic/Public/SubsonicBuiltInParameters.h
@@ -0,0 +1,15 @@
+// Copyright Epic Games, Inc. All Rights Reserved.
+
+#pragma once
+
+#include "UObject/NameTypes.h"
+
+namespace UE::Subsonic::BuiltInParameters
+{
+	extern const FName Volume;
+	extern const FName PitchShift;
+	extern const FName HighpassCutoff;
+	extern const FName LowpassCutoff;
+	extern const FName FadeOutTime;
+
+} // namespace UE::Subsonic::BuiltInParameters
diff --git a/Engine/Plugins/Experimental/Subsonic/Source/Subsonic/Public/SubsonicExecutor.h b/Engine/Plugins/Experimental/Subsonic/Source/Subsonic/Public/SubsonicExecutor.h
--- a/Engine/Plugins/Experimental/Subsonic/Source/Subsonic/Public/SubsonicExecutor.h
+++ b/Engine/Plugins/Experimental/Subsonic/Source/Subsonic/Public/SubsonicExecutor.h
@@ -6,6 +6,7 @@
 #include "HAL/Platform.h"
 #include "Misc/CoreDefines.h"
 #include "SubsonicHandles.h"
+#include "SubsonicParameterStore.h"
 #include "Templates/TypeHash.h"
 #include "Templates/UniquePtr.h"
 #include "UObject/NameTypes.h"
@@ -84,6 +85,13 @@ namespace UE::Subsonic
 		// local executor in invalid state.
 		UE_API void Unregister();
 
+		// Sets trigger-time parameters on this executor. Actions will merge these with their
+		// authored parameters when Execute() is called. Cleared on Unregister().
+		UE_API void SetParameters(FSubsonicParameterStore InParams);
+
+		// Returns the current trigger-time parameter store.
+		UE_API const FSubsonicParameterStore& GetParameters() const;
+
 	private:
 		// Registers Executor with the SubsonicSubsystem (and by extension all subscribers).
 		void Register();
@@ -96,6 +104,9 @@ namespace UE::Subsonic
 
 		// Id of calling context executor.
 		uint32 ExecutorId = GetInvalidId();
+
+		// Trigger-time parameters stored here before ExecuteEvent is called.
+		FSubsonicParameterStore Parameters;
 	};
 
 	// Key consisting of executor and scoped handle (ex. CollectionHandle, EventHandle, ActionHandle).
diff --git a/Engine/Plugins/Experimental/Subsonic/Source/Subsonic/Public/SubsonicParameterStore.h b/Engine/Plugins/Experimental/Subsonic/Source/Subsonic/Public/SubsonicParameterStore.h
new file mode 100644
--- /dev/null
+++ b/Engine/Plugins/Experimental/Subsonic/Source/Subsonic/Public/SubsonicParameterStore.h
@@ -0,0 +1,35 @@
+// Copyright Epic Games, Inc. All Rights Reserved.
+
+#pragma once
+
+#include "StructUtils/PropertyBag.h"
+
+#include "SubsonicParameterStore.generated.h"
+
+#define UE_API SUBSONIC_API
+
+
+namespace UE::Subsonic
+{
+	// A parameter store for Subsonic runtime parameters, backed by an FInstancedPropertyBag.
+	// Priority layers (lowest to highest): Authored < TriggerTime < Runtime.
+	USTRUCT(BlueprintType)
+	struct FSubsonicParameterStore
+	{
+		GENERATED_BODY()
+
+		UE_API bool HasParameter(FName Name) const;
+		UE_API void RemoveParameter(FName Name);
+		UE_API void Reset();
+		UE_API bool IsEmpty() const;
+
+		// Overlay: properties in Other are added to this store and their values overwrite any
+		// existing values with the same name.
+		UE_API void MergeFrom(const FSubsonicParameterStore& Other);
+
+		UPROPERTY(EditAnywhere, Category = "Parameters")
+		FInstancedPropertyBag Bag;
+	};
+} // namespace UE::Subsonic
+
+#undef UE_API
diff --git a/Engine/Plugins/Experimental/Subsonic/Source/Subsonic/Public/SubsonicSubsystem.h b/Engine/Plugins/Experimental/Subsonic/Source/Subsonic/Public/SubsonicSubsystem.h
--- a/Engine/Plugins/Experimental/Subsonic/Source/Subsonic/Public/SubsonicSubsystem.h
+++ b/Engine/Plugins/Experimental/Subsonic/Source/Subsonic/Public/SubsonicSubsystem.h
@@ -42,7 +42,7 @@ namespace UE::Subsonic
 		UE_API UPARAM(DisplayName = "Executor") USubsonicEventCollectionExecutor* CreateExecutorBP(UObject* WorldContextObject, FName Name, const USubsonicEventCollection* Collection);
 
 		UFUNCTION(BlueprintCallable, Category = "Subsonic", meta = (DisplayName = "Execute Event"))
-		UE_API void ExecuteEventBP(const USubsonicEventCollectionExecutor* Executor, FGameplayTag Event);
+		UE_API void ExecuteEventBP(USubsonicEventCollectionExecutor* Executor, FGameplayTag Event, FSubsonicParameterStore TriggerParams);
 
 	private:
 		friend struct FSubsonicExecutor;
diff --git a/Engine/Plugins/Experimental/Subsonic/Source/Subsonic/Subsonic.Build.cs b/Engine/Plugins/Experimental/Subsonic/Source/Subsonic/Subsonic.Build.cs
--- a/Engine/Plugins/Experimental/Subsonic/Source/Subsonic/Subsonic.Build.cs
+++ b/Engine/Plugins/Experimental/Subsonic/Source/Subsonic/Subsonic.Build.cs
@@ -7,7 +7,7 @@ public class Subsonic : ModuleRules
 	public Subsonic(ReadOnlyTargetRules Target) : base(Target)
 	{
 		PCHUsage = ModuleRules.PCHUsageMode.UseExplicitOrSharedPCHs;
-		
+
 		PublicDependencyModuleNames.AddRange(
 			new string[]
 			{
@@ -17,7 +17,7 @@ public class Subsonic : ModuleRules
 				"PropertyBindingUtils"
 			}
 		);
-			
+
 		PrivateDependencyModuleNames.AddRange(
 			new string[]
 			{
@@ -25,6 +25,10 @@ public class Subsonic : ModuleRules
 				"AudioMixer",
 				"CoreUObject",
 				"Engine",
+				"MetasoundEngine",
+				"MetasoundFrontend",
+				"MetasoundGenerator",
+				"MetasoundGraphCore",
 				"SignalProcessing",
 			}
 		);
diff --git a/Engine/Plugins/Experimental/Subsonic/Source/SubsonicEditor/Private/SubsonicEventCollectionEditor.cpp b/Engine/Plugins/Experimental/Subsonic/Source/SubsonicEditor/Private/SubsonicEventCollectionEditor.cpp
--- a/Engine/Plugins/Experimental/Subsonic/Source/SubsonicEditor/Private/SubsonicEventCollectionEditor.cpp
+++ b/Engine/Plugins/Experimental/Subsonic/Source/SubsonicEditor/Private/SubsonicEventCollectionEditor.cpp
@@ -5,22 +5,25 @@
 #include "Algo/AnyOf.h"
 #include "AudioDeviceHandle.h"
 #include "AudioDeviceManager.h"
+#include "CoreActions/SubsonicAction_GeneratorSource.h"
+#include "CoreEventSubscribers/SubsonicGeneratorSourceSubscriber.h"
 #include "Engine/Engine.h"
 #include "Framework/Application/SlateApplication.h"
 #include "Framework/Docking/TabManager.h"
 #include "GameplayTagContainer.h"
-#include "Input/DragAndDrop.h"
 #include "ISinglePropertyView.h"
+#include "Input/DragAndDrop.h"
 #include "Misc/MessageDialog.h"
 #include "Modules/ModuleManager.h"
 #include "PropertyEditorModule.h"
-#include "ScopedTransaction.h"
 #include "SGameplayTagCombo.h"
+#include "ScopedTransaction.h"
 #include "Styling/SlateStyleRegistry.h"
 #include "SubsonicBindingUtils.h"
 #include "SubsonicEditorSubsystem.h"
 #include "SubsonicEventCollection.h"
 #include "SubsonicEventCollectionEditorCommands.h"
+#include "SubsonicExecutor.h"
 #include "SubsonicPropertyBindingExtension.h"
 #include "Widgets/Docking/SDockTab.h"
 #include "Widgets/Input/SSearchableComboBox.h"
@@ -1266,19 +1269,78 @@ namespace UE::Subsonic::Editor
 		return FEventCollectionEditor::EditorName;
 	}
 
-	void FEventCollectionEditor::PostUndo(bool bSuccess)
-	{
-		FSlateApplication::Get().DismissAllMenus();
-		RebuildEventTree(true /* bRebuildListView */);
-		FEditorUndoClient::PostUndo(bSuccess);
-	}
-
 	void FEventCollectionEditor::NotifyPostChange(const FPropertyChangedEvent& PropertyChangedEvent, FProperty* PropertyThatChanged)
 	{
 		if (EventTreeInspectorView.IsValid())
 		{
 			EventTreeInspectorView->ForceRefresh();
 		}
+
+		if (!PreviewExecutor.IsValid() || !TreeDetailsView.IsValid() || TreeDetailsView->ActionIndex < 0)
+		{
+			return;
+		}
+
+		const USubsonicEventCollection* Collection = GetEventCollection();
+		if (!Collection)
+		{
+			return;
+		}
+
+		const FCollectionHandle CollectionHandle = Collection->GetHandle();
+		const FActionHandle ActionHandle = {
+			.Event = { .Collection = CollectionHandle, .EventName = TreeDetailsView->Event },
+			.Index = TreeDetailsView->ActionIndex
+		};
+
+		const FSubsonicEventActionDefinition* ActionDef = Collection->GetDefinition().FindAction(ActionHandle);
+		if (!ActionDef)
+		{
+			return;
+		}
+
+		const auto* PlayAction = ActionDef->GetAction().GetPtr<FSubsonicEventAction_GeneratorSourcePlay>();
+		if (!PlayAction)
+		{
+			return;
+		}
+
+		const FSubsonicExecutor& Executor = PreviewExecutor->GetExecutor();
+		const FActionScopeKey ScopeKey(Executor, ActionHandle);
+		USubsonicGeneratorSourceSubscriber* Subscriber = ScopeKey.FindDeviceSubsystem<USubsonicGeneratorSourceSubscriber>();
+		if (!Subscriber)
+		{
+			return;
+		}
+
+		switch (PlayAction->Scope)
+		{
+		case ESubsonicExecutionScope::Global:
+			Subscriber->SetAuthoredParameters(PlayAction->Name, PlayAction->Parameters);
+			break;
+
+		case ESubsonicExecutionScope::Collection:
+			Subscriber->SetAuthoredParameters(ActionHandle.Event.Collection, PlayAction->Name, PlayAction->Parameters);
+			break;
+
+		case ESubsonicExecutionScope::Event:
+			Subscriber->SetAuthoredParameters(ActionHandle.Event, PlayAction->Name, PlayAction->Parameters);
+			break;
+
+		case ESubsonicExecutionScope::Action:
+			Subscriber->SetAuthoredParameters(ActionHandle, PlayAction->Name, PlayAction->Parameters);
+			break;
+
+		default:
+			checkNoEntry();
+		}
+	}
+
+	void FEventCollectionEditor::PostUndo(bool bSuccess)
+	{
+		FSlateApplication::Get().DismissAllMenus();
+		RebuildEventTree(true /* bRebuildListView */);
+		FEditorUndoClient::PostUndo(bSuccess);
 	}
 
 	void FEventCollectionEditor::RestartAudition(bool bAutoEventsEnabled)
diff --git a/Engine/Plugins/Experimental/Subsonic/Source/SubsonicEditor/SubsonicEditor.Build.cs b/Engine/Plugins/Experimental/Subsonic/Source/SubsonicEditor/SubsonicEditor.Build.cs
--- a/Engine/Plugins/Experimental/Subsonic/Source/SubsonicEditor/SubsonicEditor.Build.cs
+++ b/Engine/Plugins/Experimental/Subsonic/Source/SubsonicEditor/SubsonicEditor.Build.cs
@@ -21,6 +21,7 @@ public class SubsonicEditor : ModuleRules
 			{
 				"AssetDefinition",
 				"AudioEditor",
+				"AudioMixer",
 				"AudioWidgets",
 				"BlueprintGraph",
 				"ContentBrowser",
@@ -40,6 +41,8 @@ public class SubsonicEditor : ModuleRules
 			}
 		);
 
+		PrivateIncludePaths.Add(System.IO.Path.Combine(ModuleDirectory, "..", "Subsonic", "Private"));
+
 		bAllowUETypesInNamespaces = true;
 	}
 }
diff --git a/Engine/Plugins/Experimental/Subsonic/Subsonic.uplugin b/Engine/Plugins/Experimental/Subsonic/Subsonic.uplugin
--- a/Engine/Plugins/Experimental/Subsonic/Subsonic.uplugin
+++ b/Engine/Plugins/Experimental/Subsonic/Subsonic.uplugin
@@ -39,6 +39,10 @@
 		{
 			"Name": "PropertyBindingUtils",
 			"Enabled": true
+		},
+		{
+			"Name": "Metasound",
+			"Enabled": true
 		}
 	]
 }
diff --git a/Engine/Source/Runtime/AudioMixer/Private/AudioMixerSourceManager.cpp b/Engine/Source/Runtime/AudioMixer/Private/AudioMixerSourceManager.cpp
--- a/Engine/Source/Runtime/AudioMixer/Private/AudioMixerSourceManager.cpp
+++ b/Engine/Source/Runtime/AudioMixer/Private/AudioMixerSourceManager.cpp
@@ -4294,7 +4294,7 @@ namespace Audio
 	void FMixerSourceManager::PlayGeneratorSource(const int32 SourceId)
 	{
 		check(SourceId >= NumLegacySources && SourceId < NumTotalSources && GeneratorSourceState[SourceId - NumLegacySources].bAllocated);
-		AudioMixerThreadMPSCCommand([this, SourceId]
+		auto Execute = [this, SourceId]()
 			{
 				FSourceInfo& SourceInfo = SourceInfos[SourceId];
 				SourceInfo.bIsDone = false;
@@ -4302,14 +4302,22 @@ namespace Audio
 				SourceInfo.bIsPaused = false;
 				SourceInfo.bIsActive = true;
 				SourceInfo.bEnableBaseSubmix = true;
-			}
-		);
+			};
+
+		if (MixerDevice->IsAudioRenderingThread())
+		{
+			Execute();
+		}
+		else
+		{
+			AudioMixerThreadMPSCCommand(MoveTemp(Execute));
+		}
 	}
 
 	void FMixerSourceManager::StopGeneratorSource(const int32 SourceId)
 	{
 		check(SourceId >= NumLegacySources && SourceId < NumTotalSources && GeneratorSourceState[SourceId - NumLegacySources].bAllocated);
-		AudioMixerThreadMPSCCommand([this, SourceId]
+		auto Execute = [this, SourceId]()
 			{
 				FSourceInfo& SourceInfo = SourceInfos[SourceId];
 				SourceInfo.bIsPlaying = false;
@@ -4320,8 +4328,16 @@ namespace Audio
 				{
 					SetSourceDone(SourceId);
 				}
-			}
-		);
+			};
+
+		if (MixerDevice->IsAudioRenderingThread())
+		{
+			Execute();
+		}
+		else
+		{
+			AudioMixerThreadMPSCCommand(MoveTemp(Execute));
+		}
 	}
 
 	bool FMixerSourceManager::IsGeneratorSourceDone(const int32 SourceId)
