Subsonic Event Collection Editor UX Updates
Event Tree: SListView to STreeView migration
Play button flash on test execute
Auto Audition arm button

#rnx

#Synced-CL 52327832

diff --git a/Engine/Plugins/Experimental/Subsonic/Source/Subsonic/Private/SubsonicEventCollection.cpp b/Engine/Plugins/Experimental/Subsonic/Source/Subsonic/Private/SubsonicEventCollection.cpp
--- a/Engine/Plugins/Experimental/Subsonic/Source/Subsonic/Private/SubsonicEventCollection.cpp
+++ b/Engine/Plugins/Experimental/Subsonic/Source/Subsonic/Private/SubsonicEventCollection.cpp
@@ -266,6 +266,11 @@ namespace UE::Subsonic
 		return ActionCollection;
 	}
 
+	void FSubsonicEvent::SetAutoAudition(bool bInAutoAudition)
+	{
+		bAutoAudition = bInAutoAudition;
+	}
+
 	void FSubsonicEvent::SetIsPublic(bool bInIsPublic)
 	{
 		bIsPublic = bInIsPublic;
diff --git a/Engine/Plugins/Experimental/Subsonic/Source/Subsonic/Public/SubsonicEventCollection.h b/Engine/Plugins/Experimental/Subsonic/Source/Subsonic/Public/SubsonicEventCollection.h
--- a/Engine/Plugins/Experimental/Subsonic/Source/Subsonic/Public/SubsonicEventCollection.h
+++ b/Engine/Plugins/Experimental/Subsonic/Source/Subsonic/Public/SubsonicEventCollection.h
@@ -132,6 +132,7 @@ namespace UE::Subsonic
 
 #if WITH_EDITOR
 		UE_INTERNAL UE_API TArray<FSubsonicEventActionDefinition>& GetMutableActionCollection();
+		UE_INTERNAL UE_API void SetAutoAudition(bool bInAutoAudition);
 		UE_INTERNAL UE_API void SetIsPublic(bool bInIsPublic);
 #endif // WITH_EDITOR
 
@@ -153,8 +154,8 @@ namespace UE::Subsonic
 		bool bIsPublic = true;
 
 #if WITH_EDITORONLY_DATA
-		// If true, execute this event as soon as an audition executor is created
-		UPROPERTY(EditAnywhere, Category = "Event", meta = (EditCondition = bIsPublic))
+		// If true, execute this event as soon as an audition executor is created.
+		UPROPERTY(EditAnywhere, Transient, Category = "Event", meta = (EditCondition = bIsPublic))
 		bool bAutoAudition = false;
 
 		UPROPERTY(EditAnywhere, Category = "Event")
diff --git a/Engine/Plugins/Experimental/Subsonic/Source/SubsonicEditor/Private/SubsonicEventCollectionEditor.cpp b/Engine/Plugins/Experimental/Subsonic/Source/SubsonicEditor/Private/SubsonicEventCollectionEditor.cpp
--- a/Engine/Plugins/Experimental/Subsonic/Source/SubsonicEditor/Private/SubsonicEventCollectionEditor.cpp
+++ b/Engine/Plugins/Experimental/Subsonic/Source/SubsonicEditor/Private/SubsonicEventCollectionEditor.cpp
@@ -10,6 +10,7 @@
 #include "Engine/Engine.h"
 #include "Framework/Application/SlateApplication.h"
 #include "Framework/Docking/TabManager.h"
+#include "Framework/MultiBox/MultiBoxBuilder.h"
 #include "GameplayTagContainer.h"
 #include "ISinglePropertyView.h"
 #include "Input/DragAndDrop.h"
@@ -18,6 +19,7 @@
 #include "PropertyEditorModule.h"
 #include "SGameplayTagCombo.h"
 #include "ScopedTransaction.h"
+#include "Styling/StyleColors.h"
 #include "Styling/SlateStyleRegistry.h"
 #include "SubsonicBindingUtils.h"
 #include "SubsonicEditorSubsystem.h"
@@ -26,6 +28,7 @@
 #include "SubsonicExecutor.h"
 #include "SubsonicPropertyBindingExtension.h"
 #include "Widgets/Docking/SDockTab.h"
+#include "Widgets/Input/SComboButton.h"
 #include "Widgets/Input/SSearchableComboBox.h"
 #include "Widgets/Input/SSearchBox.h"
 #include "Widgets/Layout/SBox.h"
@@ -51,6 +54,8 @@ namespace UE::Subsonic::Editor
 		const FName TabName_Inspector = "SubsonicEventCollectionEditor_Inspector";
 		const FName TabName_Parameters = "SubsonicEventCollectionEditor_Parameters";
 
+		constexpr double IconLerpDuration = 0.5;
+
 		class FActionDragDropOp : public FDragDropOperation
 		{
 		public:
@@ -226,14 +231,12 @@ namespace UE::Subsonic::Editor
 			return NewRow;
 		}
 
-		constexpr float IndentMult = 4.f;
 		FText RowDisplayName = FText::FromName(NAME_None); // Valid display name for actions with no struct type set
-		float RowIndent = 5.f;
+		constexpr float RowIndent = 5.f;
 		if (HandleVariant->IsAction())
 		{
 			if (const FSubsonicEventActionDefinition* Action = Collection->GetDefinition().FindAction(HandleVariant->GetAction()))
 			{
-				RowIndent *= IndentMult;
 				if (const UScriptStruct* ActionStruct = Action->GetAction().GetScriptStruct())
 				{
 					RowDisplayName = ActionStruct->GetDisplayNameText();
@@ -376,13 +379,13 @@ namespace UE::Subsonic::Editor
 			NameWidget.ToSharedRef()
 		];
 
-		auto VisibleOnRowHover = [RowPtr = NewRow.ToWeakPtr()]()
+		auto IsHovered = [RowPtr = NewRow.ToWeakPtr()]()
 		{
 			if (RowPtr.IsValid() && RowPtr.Pin()->IsHovered())
 			{
-				return EVisibility::Visible;
+				return true;
 			}
-			return EVisibility::Hidden;
+			return false;
 		};
 		TSharedRef<FHandleVariant> HandleRef = HandleVariant.ToSharedRef();
 
@@ -441,6 +444,76 @@ namespace UE::Subsonic::Editor
 			NameBox->AddSlot()
 			.VAlign(VAlign_Center)
 			.AutoWidth()
+			[
+				SNew(SButton)
+				.ButtonStyle(FAppStyle::Get(), "SimpleButton")
+				.OnClicked_Lambda([this, HandleRef]()
+				{
+					TransactEvent(HandleRef->GetEvent(), LOCTEXT("ToggleAutoAudition", "Toggle event {0} auto audition"), [](FSubsonicEvent& Event)
+					{
+						Event.SetAutoAudition(!Event.GetAutoAudition());
+					});
+					return FReply::Handled();
+				})
+				.ToolTipText_Lambda([this, HandleRef]()
+				{
+					if (const USubsonicEventCollection* Collection = GetEventCollection())
+					{
+						if (const FSubsonicEvent* Event = Collection->GetDefinition().FindEvent(HandleRef->GetEvent()))
+						{
+							if (Event->GetAutoAudition())
+							{
+								return LOCTEXT("AutoAuditionEnabled_Tooltip", "Auto Audition is enabled. Event will execute automatically when audition starts.");
+							}
+						}
+					}
+					return LOCTEXT("AutoAuditionDisabled_Tooltip", "Auto Audition is disabled. Click to enable.");
+				})
+				.Visibility_Lambda([this, HandleRef, IsHovered]()
+				{
+					if (IsHovered())
+					{
+						return EVisibility::Visible;
+					}
+
+					if (const USubsonicEventCollection* Collection = GetEventCollection())
+					{
+						if (const FSubsonicEvent* Event = Collection->GetDefinition().FindEvent(HandleRef->GetEvent()))
+						{
+							if (Event->GetAutoAudition())
+							{
+								return EVisibility::Visible;
+							}
+						}
+					}
+					return EVisibility::Hidden;
+				})
+				.ForegroundColor(FSlateColor::UseForeground())
+				.ContentPadding(0)
+				.IsFocusable(false)
+				[
+					SNew(SImage)
+					.Image(&EventCollectionEditorPrivate::GetSlateBrushSafe("SubsonicEditor.Arm.Small"))
+					.ColorAndOpacity_Lambda([this, HandleRef]() -> FSlateColor
+					{
+						if (const USubsonicEventCollection* Collection = GetEventCollection())
+						{
+							if (const FSubsonicEvent* Event = Collection->GetDefinition().FindEvent(HandleRef->GetEvent()))
+							{
+								if (Event->GetAutoAudition())
+								{
+									return FSlateColor(FLinearColor::Red);
+								}
+							}
+						}
+						return FSlateColor::UseForeground();
+					})
+					.DesiredSizeOverride(FVector2D(16.f, 16.f))
+				]
+			];
+			NameBox->AddSlot()
+			.VAlign(VAlign_Center)
+			.AutoWidth()
 			[
 				SNew(SButton)
 				.ButtonStyle(FAppStyle::Get(), "SimpleButton")
@@ -456,9 +529,38 @@ namespace UE::Subsonic::Editor
 				[
 					SNew(SImage)
 					.Image(&EventCollectionEditorPrivate::GetSlateBrushSafe("SubsonicEditor.Play.Small"))
-					.ColorAndOpacity(FSlateColor::UseForeground())
+					.ColorAndOpacity_Lambda([this, EventName = HandleRef->GetEvent().EventName]() -> FSlateColor
+					{
+						if (const double* LerpTime = EventExecuteLerpTimes.Find(EventName))
+						{
+							const double Elapsed = FMath::Max(0.0, FSlateApplication::Get().GetCurrentTime() - *LerpTime);
+							if (Elapsed < EventCollectionEditorPrivate::IconLerpDuration)
+							{
+								const float Alpha = static_cast<float>(Elapsed / EventCollectionEditorPrivate::IconLerpDuration);
+								const FLinearColor ForegroundColor = FStyleColors::Foreground.GetSpecifiedColor();
+								return FSlateColor(FLinearColor::LerpUsingHSV(FLinearColor::Green, ForegroundColor, Alpha));
+							}
+						}
+						return FSlateColor::UseForeground();
+					})
 					.DesiredSizeOverride(FVector2D(16.f, 16.f))
-					.Visibility_Lambda(VisibleOnRowHover)
+					.Visibility_Lambda([this, EventName = HandleRef->GetEvent().EventName, IsHovered]()
+					{
+						if (IsHovered())
+						{
+							return EVisibility::Visible;
+						}
+
+						if (const double* LerpTime = EventExecuteLerpTimes.Find(EventName))
+						{
+							const double Elapsed = FMath::Max(0.0, FSlateApplication::Get().GetCurrentTime() - *LerpTime);
+							if (Elapsed < EventCollectionEditorPrivate::IconLerpDuration)
+							{
+								return EVisibility::Visible;
+							}
+						}
+						return EVisibility::Hidden;
+					})
 				]
 			];
 		}
@@ -503,7 +605,7 @@ namespace UE::Subsonic::Editor
 				.ForegroundColor(FSlateColor::UseForeground())
 				.ContentPadding(0)
 				.IsFocusable(false)
-				.Visibility_Lambda(VisibleOnRowHover)
+				.Visibility_Lambda([IsHovered]() { return IsHovered() ? EVisibility::Visible : EVisibility::Hidden; })
 				[
 					SNew(SImage)
 					.Image(FAppStyle::Get().GetBrush("Icons.AddCircle"))
@@ -536,7 +638,7 @@ namespace UE::Subsonic::Editor
 					: LOCTEXT("DeleteActionToolTip_Format", "Delete action from event"))
 				.ForegroundColor(FSlateColor::UseForeground())
 				.ContentPadding(0)
-				.Visibility_Lambda(VisibleOnRowHover)
+				.Visibility_Lambda([IsHovered]() { return IsHovered() ? EVisibility::Visible : EVisibility::Hidden; })
 				.IsFocusable(false)
 				[
 					SNew(SImage)
@@ -549,6 +651,17 @@ namespace UE::Subsonic::Editor
 		return NewRow;
 	}
 
+	void FEventCollectionEditor::OnEventTree_GetChildren(TSharedPtr<FHandleVariant> InParent, TArray<TSharedPtr<FHandleVariant>>& OutChildren)
+	{
+		if (InParent.IsValid() && InParent->IsEvent())
+		{
+			if (const TArray<TSharedPtr<FHandleVariant>>* Children = EventTreeChildHandles.Find(InParent->GetEvent().EventName))
+			{
+				OutChildren = *Children;
+			}
+		}
+	}
+
 	void FEventCollectionEditor::OnEventTree_RowSelected(TSharedPtr<FHandleVariant> HandleVariant, ESelectInfo::Type SelectInfo)
 	{
 		if (TreeDetailsView.IsValid() && HandleVariant.IsValid())
@@ -578,7 +691,7 @@ namespace UE::Subsonic::Editor
 		{
 			return;
 		}
-		for (const TSharedPtr<FHandleVariant>& Handle : EventTreeHandles)
+		for (const TSharedPtr<FHandleVariant>& Handle : EventTreeRootHandles)
 		{
 			if (Handle.IsValid() && Handle->IsEvent() && Handle->GetEvent().EventName == EventName)
 			{
@@ -588,6 +701,28 @@ namespace UE::Subsonic::Editor
 		}
 	}
 
+	void FEventCollectionEditor::ExpandAllEvents()
+	{
+		if (EventTreeView.IsValid())
+		{
+			for (const TSharedPtr<FHandleVariant>& RootHandle : EventTreeRootHandles)
+			{
+				EventTreeView->SetItemExpansion(RootHandle, true);
+			}
+		}
+	}
+
+	void FEventCollectionEditor::CollapseAllEvents()
+	{
+		if (EventTreeView.IsValid())
+		{
+			for (const TSharedPtr<FHandleVariant>& RootHandle : EventTreeRootHandles)
+			{
+				EventTreeView->SetItemExpansion(RootHandle, false);
+			}
+		}
+	}
+
 	void FEventCollectionEditor::TryInvokeTab(FName TabName) const
 	{
 		if (TabManager.IsValid())
@@ -599,7 +734,8 @@ namespace UE::Subsonic::Editor
 	void FEventCollectionEditor::RebuildEventTree(bool bRebuildListView)
 	{
 		ActionComboValues.Reset();
-		EventTreeHandles.Reset();
+		EventTreeRootHandles.Reset();
+		EventTreeChildHandles.Reset();
 
 		if (!GEditor)
 		{
@@ -612,6 +748,28 @@ namespace UE::Subsonic::Editor
 			return;
 		}
 
+		// Remove Lerp entries for events that no longer exist (renamed or deleted)
+		if (EventExecuteLerpTimes.Num() > 0)
+		{
+			const TMap<FGameplayTag, FSubsonicEvent>& Events = EventCollection->GetDefinition().GetEvents();
+			for (auto It = EventExecuteLerpTimes.CreateIterator(); It; ++It)
+			{
+				bool bFound = false;
+				for (const auto& EventPair : Events)
+				{
+					if (EventPair.Key.GetTagName() == It->Key)
+					{
+						bFound = true;
+						break;
+					}
+				}
+				if (!bFound)
+				{
+					It.RemoveCurrent();
+				}
+			}
+		}
+
 		FPropertyEditorModule& PropertyEditor = FModuleManager::LoadModuleChecked<FPropertyEditorModule>("PropertyEditor");
 		CollectionView = PropertyEditor.CreateSingleProperty(EventCollection, USubsonicEventCollection::GetDefinitionPropertyName(), { });
 
@@ -660,21 +818,29 @@ namespace UE::Subsonic::Editor
 			FHandleVariant EventVariant;
 			const FEventHandle EventHandle = { .Collection = EventCollection->GetHandle(), .EventName = Pair.Key.GetTagName() };
 			EventVariant.Set(EventHandle);
-			EventTreeHandles.Add(MakeShared<FHandleVariant>(MoveTemp(EventVariant)));
+			TSharedPtr<FHandleVariant> EventPtr = MakeShared<FHandleVariant>(MoveTemp(EventVariant));
+			EventTreeRootHandles.Add(EventPtr);
 
+			TArray<TSharedPtr<FHandleVariant>>& ChildHandles = EventTreeChildHandles.Add(Pair.Key.GetTagName());
 			int32 ActionIndex = 0;
 			for (const FSubsonicEventActionDefinition& ActionDefinition : Pair.Value->GetActionCollection())
 			{
 				FHandleVariant ActionVariant;
 				const FActionHandle ActionHandle = { .Event = EventHandle, .Index = ActionIndex++ };
 				ActionVariant.Set(ActionHandle);
-				EventTreeHandles.Add(MakeShared<FHandleVariant>(MoveTemp(ActionVariant)));
+				ChildHandles.Add(MakeShared<FHandleVariant>(MoveTemp(ActionVariant)));
 			}
 		}
 
 		if (bRebuildListView && EventTreeView.IsValid())
 		{
-			EventTreeView->RebuildList();
+			EventTreeView->RequestTreeRefresh();
+
+			// Default all events to expanded
+			for (const TSharedPtr<FHandleVariant>& RootHandle : EventTreeRootHandles)
+			{
+				EventTreeView->SetItemExpansion(RootHandle, true);
+			}
 		}
 	}
 
@@ -878,7 +1044,7 @@ namespace UE::Subsonic::Editor
 
 		InitBindings();
 
-		RebuildEventTree();
+		RebuildEventTree(true /* bRebuildListView */);
 
 		if (ParametersDetailsView.IsValid())
 		{
@@ -987,6 +1153,41 @@ namespace UE::Subsonic::Editor
 						.ColorAndOpacity(FSlateColor::UseForeground())
 					]
 				]
+				+ SHorizontalBox::Slot()
+				.VAlign(VAlign_Center)
+				.HAlign(HAlign_Right)
+				.AutoWidth()
+				[
+					SNew(SComboButton)
+					.ComboButtonStyle(FAppStyle::Get(), "SimpleComboButton")
+					.HasDownArrow(false)
+					.ContentPadding(0)
+					.ForegroundColor(FSlateColor::UseForeground())
+					.ToolTipText(LOCTEXT("EventTreeOptionsToolTip", "Event tree options"))
+					.OnGetMenuContent_Lambda([this]()
+					{
+						FMenuBuilder MenuBuilder(true, nullptr);
+						MenuBuilder.AddMenuEntry(
+							LOCTEXT("ExpandAllEvents", "Expand All"),
+							LOCTEXT("ExpandAllEvents_Tooltip", "Expand all events in the tree"),
+							FSlateIcon(),
+							FUIAction(FExecuteAction::CreateSP(this, &FEventCollectionEditor::ExpandAllEvents))
+						);
+						MenuBuilder.AddMenuEntry(
+							LOCTEXT("CollapseAllEvents", "Collapse All"),
+							LOCTEXT("CollapseAllEvents_Tooltip", "Collapse all events in the tree"),
+							FSlateIcon(),
+							FUIAction(FExecuteAction::CreateSP(this, &FEventCollectionEditor::CollapseAllEvents))
+						);
+						return MenuBuilder.MakeWidget();
+					})
+					.ButtonContent()
+					[
+						SNew(SImage)
+						.Image(FAppStyle::Get().GetBrush("Icons.Filter"))
+						.ColorAndOpacity(FSlateColor::UseForeground())
+					]
+				]
 			]
 		);
 
@@ -1013,9 +1214,10 @@ namespace UE::Subsonic::Editor
 		+ SVerticalBox::Slot()
 		.AutoHeight()
 		[
-			SAssignNew(EventTreeView, SListView<TSharedPtr<FHandleVariant>>)
+			SAssignNew(EventTreeView, STreeView<TSharedPtr<FHandleVariant>>)
 			.HeaderRow(EventTreeHeader)
-			.ListItemsSource(&EventTreeHandles)
+			.TreeItemsSource(&EventTreeRootHandles)
+			.OnGetChildren(this, &FEventCollectionEditor::OnEventTree_GetChildren)
 			.OnGenerateRow(this, &FEventCollectionEditor::OnEventTree_GenerateRow)
 			.OnSelectionChanged(this, &FEventCollectionEditor::OnEventTree_RowSelected)
 			.SelectionMode(ESelectionMode::Single)
@@ -1261,6 +1463,33 @@ namespace UE::Subsonic::Editor
 		if (PreviewExecutor.IsValid())
 		{
 			PreviewExecutor->GetExecutor().ExecuteEvent(EventHandle.EventName);
+
+			// Record Lerp start time for the play button green Lerp animation
+			EventExecuteLerpTimes.Add(EventHandle.EventName, FSlateApplication::Get().GetCurrentTime());
+
+			// Register an active timer on the event tree to force repaints during the Lerp
+			if (EventTreeView.IsValid())
+			{
+				EventTreeView->RegisterActiveTimer(0.0f, FWidgetActiveTimerDelegate::CreateLambda(
+					[this](double, float) -> EActiveTimerReturnType
+					{
+						const double Now = FSlateApplication::Get().GetCurrentTime();
+						bool bAnyActive = false;
+						for (auto It = EventExecuteLerpTimes.CreateIterator(); It; ++It)
+						{
+							if (Now - It->Value < EventCollectionEditorPrivate::IconLerpDuration)
+							{
+								bAnyActive = true;
+							}
+							else
+							{
+								It.RemoveCurrent();
+							}
+						}
+						return bAnyActive ? EActiveTimerReturnType::Continue : EActiveTimerReturnType::Stop;
+					}
+				));
+			}
 		}
 	}
 
@@ -1362,14 +1591,40 @@ namespace UE::Subsonic::Editor
 
 					if (bAutoEventsEnabled)
 					{
+						const double LerpTime = FSlateApplication::Get().GetCurrentTime();
 						const FSubsonicEventCollectionDefinition& CollectionDefinition = EventCollection->GetDefinition();
 						for (const TPair<FGameplayTag, FSubsonicEvent>& Pair : CollectionDefinition.GetEvents())
 						{
 							if (Pair.Value.GetAutoAudition())
 							{
 								PreviewExecutor->GetExecutor().ExecuteEvent(Pair.Key.GetTagName());
+								EventExecuteLerpTimes.Add(Pair.Key.GetTagName(), LerpTime);
 							}
 						}
+
+						// Register an active timer to repaint during the Lerp
+						if (EventExecuteLerpTimes.Num() > 0 && EventTreeView.IsValid())
+						{
+							EventTreeView->RegisterActiveTimer(0.0f, FWidgetActiveTimerDelegate::CreateLambda(
+								[this](double, float) -> EActiveTimerReturnType
+								{
+									const double Now = FSlateApplication::Get().GetCurrentTime();
+									bool bAnyActive = false;
+									for (auto It = EventExecuteLerpTimes.CreateIterator(); It; ++It)
+									{
+										if (Now - It->Value < EventCollectionEditorPrivate::IconLerpDuration)
+										{
+											bAnyActive = true;
+										}
+										else
+										{
+											It.RemoveCurrent();
+										}
+									}
+									return bAnyActive ? EActiveTimerReturnType::Continue : EActiveTimerReturnType::Stop;
+								}
+							));
+						}
 					}
 				}
 			}
diff --git a/Engine/Plugins/Experimental/Subsonic/Source/SubsonicEditor/Private/SubsonicEventCollectionEditor.h b/Engine/Plugins/Experimental/Subsonic/Source/SubsonicEditor/Private/SubsonicEventCollectionEditor.h
--- a/Engine/Plugins/Experimental/Subsonic/Source/SubsonicEditor/Private/SubsonicEventCollectionEditor.h
+++ b/Engine/Plugins/Experimental/Subsonic/Source/SubsonicEditor/Private/SubsonicEventCollectionEditor.h
@@ -18,6 +18,7 @@
 #include "UObject/StrongObjectPtrTemplates.h"
 #include "Widgets/Views/SListView.h"
 #include "Widgets/Views/STableRow.h"
+#include "Widgets/Views/STreeView.h"
 
 
 namespace UE::Subsonic
@@ -125,8 +126,12 @@ namespace UE::Subsonic
 			TSharedRef<SHandleVariantTableRow> InitDraggableEventRow(TSharedPtr<FHandleVariant> HandleVariant, const TSharedRef<STableViewBase>& TableView);
 
 			TSharedRef<ITableRow> OnEventTree_GenerateRow(TSharedPtr<FHandleVariant> HandleVariant, const TSharedRef<STableViewBase>& TableView);
+			void OnEventTree_GetChildren(TSharedPtr<FHandleVariant> InParent, TArray<TSharedPtr<FHandleVariant>>& OutChildren);
 			void OnEventTree_RowSelected(TSharedPtr<FHandleVariant> HandleVariant, ESelectInfo::Type SelectInfo);
 
+			void ExpandAllEvents();
+			void CollapseAllEvents();
+
 			void RebuildEventTree(bool bRebuildListView = false);
 
 			void RestartAudition(bool bAutoEventsEnabled = true);
@@ -146,8 +151,9 @@ namespace UE::Subsonic
 			TSharedPtr<FTabManager::FStack> ParametersStack;
 			TStrongObjectPtr<USubsonicCollectionParametersView> ParametersView;
 
-			TArray<TSharedPtr<FHandleVariant>> EventTreeHandles;
-			TSharedPtr<SListView<TSharedPtr<FHandleVariant>>> EventTreeView;
+			TArray<TSharedPtr<FHandleVariant>> EventTreeRootHandles;
+			TMap<FName, TArray<TSharedPtr<FHandleVariant>>> EventTreeChildHandles;
+			TSharedPtr<STreeView<TSharedPtr<FHandleVariant>>> EventTreeView;
 			FString EventViewFilter;
 
 			// Value names of selectable actions in event tree view
@@ -158,6 +164,10 @@ namespace UE::Subsonic
 			TSharedPtr<ISinglePropertyView> CollectionView;
 
 			FDelegateHandle StaleBindingsDelegateHandle;
+
+			// Tracks the last time each event was test-executed via the play button,
+			// used to Lerp the button green and fade back to white.
+			TMap<FName, double> EventExecuteLerpTimes;
 		};
 	} // namespace Editor
 } // namespace UE::Subsonic
\ No newline at end of file
diff --git a/Engine/Plugins/Experimental/Subsonic/Source/SubsonicEditor/Private/SubsonicSlateStyle.h b/Engine/Plugins/Experimental/Subsonic/Source/SubsonicEditor/Private/SubsonicSlateStyle.h
--- a/Engine/Plugins/Experimental/Subsonic/Source/SubsonicEditor/Private/SubsonicSlateStyle.h
+++ b/Engine/Plugins/Experimental/Subsonic/Source/SubsonicEditor/Private/SubsonicSlateStyle.h
@@ -27,6 +27,7 @@ namespace UE::Subsonic::Editor
 			{
 				Set("SubsonicEventCollection.Color", FColor(156, 0, 255));
 
+				Set("SubsonicEditor.Arm.Small", new IMAGE_BRUSH_SVG(TEXT("Icons/arm"), Icon20));
 				Set("SubsonicEditor.Play", new IMAGE_BRUSH_SVG(TEXT("Icons/play"), Icon40));
 				Set("SubsonicEditor.Play.Small", new IMAGE_BRUSH_SVG(TEXT("Icons/play"), Icon20));
 				Set("SubsonicEditor.Play.Thumbnail", new IMAGE_BRUSH_SVG(TEXT("Icons/play_thumbnail"), Icon64));
