Split Subsonic Module into two modules, Engine (UObject wrappers) & Core
- Move SubsonicSubsystem to implemented implementation in Engine that is statically accessible post SubsonicSubsystem accessor wrapper is initialized.
- Move ExecuteEvent from Subsystem to Executor UObject wrapper for better visibility

#rnx

#Synced-CL 52398802

diff --git a/Engine/Plugins/Experimental/Subsonic/Source/Subsonic/Private/CoreActions/SubsonicAction_AudioComponent.cpp b/Engine/Plugins/Experimental/Subsonic/Source/Subsonic/Private/CoreActions/SubsonicAction_AudioComponent.cpp
deleted file mode 100644
diff --git a/Engine/Plugins/Experimental/Subsonic/Source/Subsonic/Private/CoreActions/SubsonicAction_AudioComponent.h b/Engine/Plugins/Experimental/Subsonic/Source/Subsonic/Private/CoreActions/SubsonicAction_AudioComponent.h
deleted file mode 100644
diff --git a/Engine/Plugins/Experimental/Subsonic/Source/Subsonic/Private/CoreActions/SubsonicAction_EventCore.cpp b/Engine/Plugins/Experimental/Subsonic/Source/Subsonic/Private/CoreActions/SubsonicAction_EventCore.cpp
deleted file mode 100644
diff --git a/Engine/Plugins/Experimental/Subsonic/Source/Subsonic/Private/CoreActions/SubsonicAction_EventCore.h b/Engine/Plugins/Experimental/Subsonic/Source/Subsonic/Private/CoreActions/SubsonicAction_EventCore.h
deleted file mode 100644
diff --git a/Engine/Plugins/Experimental/Subsonic/Source/Subsonic/Private/CoreActions/SubsonicAction_GeneratorSource.cpp b/Engine/Plugins/Experimental/Subsonic/Source/Subsonic/Private/CoreActions/SubsonicAction_GeneratorSource.cpp
deleted file mode 100644
diff --git a/Engine/Plugins/Experimental/Subsonic/Source/Subsonic/Private/CoreActions/SubsonicAction_GeneratorSource.h b/Engine/Plugins/Experimental/Subsonic/Source/Subsonic/Private/CoreActions/SubsonicAction_GeneratorSource.h
deleted file mode 100644
diff --git a/Engine/Plugins/Experimental/Subsonic/Source/Subsonic/Private/CoreEventSubscribers/AudioComponentEventSubscriber.cpp b/Engine/Plugins/Experimental/Subsonic/Source/Subsonic/Private/CoreEventSubscribers/AudioComponentEventSubscriber.cpp
deleted file mode 100644
diff --git a/Engine/Plugins/Experimental/Subsonic/Source/Subsonic/Private/CoreEventSubscribers/AudioComponentEventSubscriber.h b/Engine/Plugins/Experimental/Subsonic/Source/Subsonic/Private/CoreEventSubscribers/AudioComponentEventSubscriber.h
deleted file mode 100644
diff --git a/Engine/Plugins/Experimental/Subsonic/Source/Subsonic/Private/CoreEventSubscribers/SubsonicEventSubscriberBase.cpp b/Engine/Plugins/Experimental/Subsonic/Source/Subsonic/Private/CoreEventSubscribers/SubsonicEventSubscriberBase.cpp
deleted file mode 100644
diff --git a/Engine/Plugins/Experimental/Subsonic/Source/Subsonic/Private/CoreEventSubscribers/SubsonicGenerator.cpp b/Engine/Plugins/Experimental/Subsonic/Source/Subsonic/Private/CoreEventSubscribers/SubsonicGenerator.cpp
deleted file mode 100644
diff --git a/Engine/Plugins/Experimental/Subsonic/Source/Subsonic/Private/CoreEventSubscribers/SubsonicGenerator.h b/Engine/Plugins/Experimental/Subsonic/Source/Subsonic/Private/CoreEventSubscribers/SubsonicGenerator.h
deleted file mode 100644
diff --git a/Engine/Plugins/Experimental/Subsonic/Source/Subsonic/Private/CoreEventSubscribers/SubsonicGeneratorSourceSubscriber.cpp b/Engine/Plugins/Experimental/Subsonic/Source/Subsonic/Private/CoreEventSubscribers/SubsonicGeneratorSourceSubscriber.cpp
deleted file mode 100644
diff --git a/Engine/Plugins/Experimental/Subsonic/Source/Subsonic/Private/CoreEventSubscribers/SubsonicGeneratorSourceSubscriber.h b/Engine/Plugins/Experimental/Subsonic/Source/Subsonic/Private/CoreEventSubscribers/SubsonicGeneratorSourceSubscriber.h
deleted file mode 100644
diff --git a/Engine/Plugins/Experimental/Subsonic/Source/Subsonic/Private/CoreEventSubscribers/SubsonicRelay.cpp b/Engine/Plugins/Experimental/Subsonic/Source/Subsonic/Private/CoreEventSubscribers/SubsonicRelay.cpp
deleted file mode 100644
diff --git a/Engine/Plugins/Experimental/Subsonic/Source/Subsonic/Private/CoreEventSubscribers/SubsonicRelay.h b/Engine/Plugins/Experimental/Subsonic/Source/Subsonic/Private/CoreEventSubscribers/SubsonicRelay.h
deleted file mode 100644
diff --git a/Engine/Plugins/Experimental/Subsonic/Source/Subsonic/Private/CoreEventSubscribers/SubsonicWaveGenerator.cpp b/Engine/Plugins/Experimental/Subsonic/Source/Subsonic/Private/CoreEventSubscribers/SubsonicWaveGenerator.cpp
deleted file mode 100644
diff --git a/Engine/Plugins/Experimental/Subsonic/Source/Subsonic/Private/CoreEventSubscribers/SubsonicWaveGenerator.h b/Engine/Plugins/Experimental/Subsonic/Source/Subsonic/Private/CoreEventSubscribers/SubsonicWaveGenerator.h
deleted file mode 100644
diff --git a/Engine/Plugins/Experimental/Subsonic/Source/Subsonic/Private/SubsonicBindingUtils.cpp b/Engine/Plugins/Experimental/Subsonic/Source/Subsonic/Private/SubsonicBindingUtils.cpp
deleted file mode 100644
diff --git a/Engine/Plugins/Experimental/Subsonic/Source/Subsonic/Private/SubsonicBuiltInParameters.cpp b/Engine/Plugins/Experimental/Subsonic/Source/Subsonic/Private/SubsonicBuiltInParameters.cpp
deleted file mode 100644
diff --git a/Engine/Plugins/Experimental/Subsonic/Source/Subsonic/Private/SubsonicEventCollection.cpp b/Engine/Plugins/Experimental/Subsonic/Source/Subsonic/Private/SubsonicEventCollection.cpp
deleted file mode 100644
diff --git a/Engine/Plugins/Experimental/Subsonic/Source/Subsonic/Private/SubsonicEventSystemInterface.h b/Engine/Plugins/Experimental/Subsonic/Source/Subsonic/Private/SubsonicEventSystemInterface.h
deleted file mode 100644
diff --git a/Engine/Plugins/Experimental/Subsonic/Source/Subsonic/Private/SubsonicExecutor.cpp b/Engine/Plugins/Experimental/Subsonic/Source/Subsonic/Private/SubsonicExecutor.cpp
deleted file mode 100644
diff --git a/Engine/Plugins/Experimental/Subsonic/Source/Subsonic/Private/SubsonicHandles.cpp b/Engine/Plugins/Experimental/Subsonic/Source/Subsonic/Private/SubsonicHandles.cpp
deleted file mode 100644
diff --git a/Engine/Plugins/Experimental/Subsonic/Source/Subsonic/Private/SubsonicLog.h b/Engine/Plugins/Experimental/Subsonic/Source/Subsonic/Private/SubsonicLog.h
deleted file mode 100644
diff --git a/Engine/Plugins/Experimental/Subsonic/Source/Subsonic/Private/SubsonicModule.cpp b/Engine/Plugins/Experimental/Subsonic/Source/Subsonic/Private/SubsonicModule.cpp
deleted file mode 100644
diff --git a/Engine/Plugins/Experimental/Subsonic/Source/Subsonic/Private/SubsonicParameterStore.cpp b/Engine/Plugins/Experimental/Subsonic/Source/Subsonic/Private/SubsonicParameterStore.cpp
deleted file mode 100644
diff --git a/Engine/Plugins/Experimental/Subsonic/Source/Subsonic/Private/SubsonicSubsystem.cpp b/Engine/Plugins/Experimental/Subsonic/Source/Subsonic/Private/SubsonicSubsystem.cpp
deleted file mode 100644
diff --git a/Engine/Plugins/Experimental/Subsonic/Source/Subsonic/Private/SubsonicTypes.cpp b/Engine/Plugins/Experimental/Subsonic/Source/Subsonic/Private/SubsonicTypes.cpp
deleted file mode 100644
diff --git a/Engine/Plugins/Experimental/Subsonic/Source/Subsonic/Public/CoreEventSubscribers/SubsonicEventSubscriberBase.h b/Engine/Plugins/Experimental/Subsonic/Source/Subsonic/Public/CoreEventSubscribers/SubsonicEventSubscriberBase.h
deleted file mode 100644
diff --git a/Engine/Plugins/Experimental/Subsonic/Source/Subsonic/Public/SubsonicBindingUtils.h b/Engine/Plugins/Experimental/Subsonic/Source/Subsonic/Public/SubsonicBindingUtils.h
deleted file mode 100644
diff --git a/Engine/Plugins/Experimental/Subsonic/Source/Subsonic/Public/SubsonicBuiltInParameters.h b/Engine/Plugins/Experimental/Subsonic/Source/Subsonic/Public/SubsonicBuiltInParameters.h
deleted file mode 100644
diff --git a/Engine/Plugins/Experimental/Subsonic/Source/Subsonic/Public/SubsonicEventCollection.h b/Engine/Plugins/Experimental/Subsonic/Source/Subsonic/Public/SubsonicEventCollection.h
deleted file mode 100644
diff --git a/Engine/Plugins/Experimental/Subsonic/Source/Subsonic/Public/SubsonicExecutor.h b/Engine/Plugins/Experimental/Subsonic/Source/Subsonic/Public/SubsonicExecutor.h
deleted file mode 100644
diff --git a/Engine/Plugins/Experimental/Subsonic/Source/Subsonic/Public/SubsonicHandles.h b/Engine/Plugins/Experimental/Subsonic/Source/Subsonic/Public/SubsonicHandles.h
deleted file mode 100644
diff --git a/Engine/Plugins/Experimental/Subsonic/Source/Subsonic/Public/SubsonicParameterStore.h b/Engine/Plugins/Experimental/Subsonic/Source/Subsonic/Public/SubsonicParameterStore.h
deleted file mode 100644
diff --git a/Engine/Plugins/Experimental/Subsonic/Source/Subsonic/Public/SubsonicSubsystem.h b/Engine/Plugins/Experimental/Subsonic/Source/Subsonic/Public/SubsonicSubsystem.h
deleted file mode 100644
diff --git a/Engine/Plugins/Experimental/Subsonic/Source/Subsonic/Public/SubsonicTypes.h b/Engine/Plugins/Experimental/Subsonic/Source/Subsonic/Public/SubsonicTypes.h
deleted file mode 100644
diff --git a/Engine/Plugins/Experimental/Subsonic/Source/Subsonic/Subsonic.Build.cs b/Engine/Plugins/Experimental/Subsonic/Source/Subsonic/Subsonic.Build.cs
deleted file mode 100644
diff --git a/Engine/Plugins/Experimental/Subsonic/Source/SubsonicCore/Private/ISubsonicEventRegistry.cpp b/Engine/Plugins/Experimental/Subsonic/Source/SubsonicCore/Private/ISubsonicEventRegistry.cpp
new file mode 100644
--- /dev/null
+++ b/Engine/Plugins/Experimental/Subsonic/Source/SubsonicCore/Private/ISubsonicEventRegistry.cpp
@@ -0,0 +1,31 @@
+// Copyright Epic Games, Inc. All Rights Reserved.
+
+#include "ISubsonicEventRegistry.h"
+
+namespace UE::Subsonic::Core
+{
+	TUniquePtr<ISubsonicEventRegistry> ISubsonicEventRegistry::Instance;
+
+	ISubsonicEventRegistry* ISubsonicEventRegistry::Get()
+	{
+		return Instance.Get();
+	}
+
+	ISubsonicEventRegistry& ISubsonicEventRegistry::GetChecked()
+	{
+		checkf(Instance.IsValid(), TEXT("No subsystem event registry initialized (invalid access)"));
+		return *Instance.Get();
+	}
+
+	void ISubsonicEventRegistry::Initialize(TUniquePtr<ISubsonicEventRegistry>&& InRegistry)
+	{
+		checkf(!Instance.IsValid(), TEXT("Cannot reinitialize Subsonic Event Registry (Initialize has been called more than once prior to deinitialization"));
+		Instance = MoveTemp(InRegistry);
+	}
+
+	void ISubsonicEventRegistry::Deinitialize()
+	{
+		checkf(Instance.IsValid(), TEXT("Cannot Deinitialize Subsonic Event Registry (Never initialized or Deinitialize has been called more than once"));
+		Instance.Reset();
+	}
+} // namespace UE::Subsonic::Core
diff --git a/Engine/Plugins/Experimental/Subsonic/Source/SubsonicCore/Private/SubsonicBindingUtils.cpp b/Engine/Plugins/Experimental/Subsonic/Source/SubsonicCore/Private/SubsonicBindingUtils.cpp
new file mode 100644
--- /dev/null
+++ b/Engine/Plugins/Experimental/Subsonic/Source/SubsonicCore/Private/SubsonicBindingUtils.cpp
@@ -0,0 +1,276 @@
+// Copyright Epic Games, Inc. All Rights Reserved.
+
+#if WITH_EDITOR
+#include "SubsonicBindingUtils.h"
+
+#include "Algo/AllOf.h"
+#include "Editor.h"
+#include "Misc/MessageDialog.h"
+#include "StructUtils/PropertyBag.h"
+#include "SubsonicEventCollection.h"
+#include "UObject/UnrealType.h"
+
+#define LOCTEXT_NAMESPACE "SubsonicBindingUtils"
+
+
+namespace UE::Subsonic::Core::BindingUtils
+{
+	bool ArePropertiesBindingCompatible(const FProperty* FromBag, const FProperty* ToAction)
+	{
+		check(FromBag);
+		check(ToAction);
+
+		if (ToAction->SameType(FromBag))
+		{
+			return true;
+		}
+		const FObjectPropertyBase* ActionObj = CastField<FObjectPropertyBase>(ToAction);
+		const FObjectPropertyBase* BagObj = CastField<FObjectPropertyBase>(FromBag);
+
+		return ActionObj && BagObj && BagObj->PropertyClass->IsChildOf(ActionObj->PropertyClass);
+	}
+
+#if WITH_EDITORONLY_DATA
+	FStaleBindingsMap FindStaleBindings(
+		const FSubsonicEventCollectionDefinition& Definition, const FCollectionHandle& ParentHandle)
+	{
+		FStaleBindingsMap StaleBindings;
+
+		for (const TPair<FGameplayTag, FSubsonicEvent>& EventPair : Definition.GetEvents())
+		{
+			const FSubsonicEvent& Event = EventPair.Value;
+			const FEventHandle EventHandle = { .Collection = ParentHandle, .EventName = EventPair.Key.GetTagName() };
+
+			int32 ActionIndex = 0;
+			for (const FSubsonicEventActionDefinition& ActionDef : Event.GetActionCollection())
+			{
+				const FActionHandle ActionHandle = { .Event = EventHandle, .Index = ActionIndex++ };
+				const TInstancedStruct<FSubsonicEventActionBase>& Action = ActionDef.GetAction();
+				const FSubsonicEventActionBase* ActionBase = Action.GetPtr();
+				const UScriptStruct* ActionStructType = Action.GetScriptStruct();
+				if (!ActionBase || !ActionStructType)
+				{
+					continue;
+				}
+
+				for (const TPair<FName, FSubsonicEventActionBoundProperties>& BindingPair : ActionBase->GetBindings())
+				{
+					const FName& ParamName = BindingPair.Key;
+					const FPropertyBagPropertyDesc* BagDesc = nullptr;
+					if (const UPropertyBag* EventBag = Event.GetParameters().GetPropertyBagStruct())
+					{
+						BagDesc = EventBag->FindPropertyDescByName(ParamName);
+					}
+					if (!BagDesc)
+					{
+						if (const UPropertyBag* CollBag = Definition.GetParameters().GetPropertyBagStruct())
+						{
+							BagDesc = CollBag->FindPropertyDescByName(ParamName);
+						}
+					}
+
+					const bool bNoParamFound = !BagDesc || !BagDesc->CachedProperty;
+					for (const FName& BoundPropName : BindingPair.Value.Properties)
+					{
+						if (!bNoParamFound)
+						{
+							const FProperty* ActionProp = ActionStructType->FindPropertyByName(BoundPropName);
+							if (ActionProp && BagDesc->CachedProperty && ArePropertiesBindingCompatible(BagDesc->CachedProperty, ActionProp))
+							{
+								continue;
+							}
+						}
+
+						StaleBindings.FindOrAdd(ParamName).Add(FStalePropertyHandle { .Action = ActionHandle, .Property = BoundPropName });
+					}
+				}
+			}
+		}
+
+		return StaleBindings;
+	}
+
+	FName FindRenameCandidate(const FSubsonicEventCollectionDefinition& Definition, FName OldParamName, const TArray<FStalePropertyHandle>& StaleEntries)
+	{
+		// Collect the set of action property types that were bound to the old parameter name.
+		TArray<const FProperty*> StaleActionProps;
+		for (const FStalePropertyHandle& Entry : StaleEntries)
+		{
+			if (const FSubsonicEventActionDefinition* ActionDef = Definition.FindAction(Entry.Action))
+			{
+				if (const UScriptStruct* ActionStruct = ActionDef->GetAction().GetScriptStruct())
+				{
+					if (const FProperty* ActionProp = ActionStruct->FindPropertyByName(Entry.Property))
+					{
+						StaleActionProps.AddUnique(ActionProp);
+					}
+				}
+			}
+		}
+
+		if (StaleActionProps.IsEmpty())
+		{
+			return NAME_None;
+		}
+
+		// Collect all parameter names referenced by ANY binding across the ENTIRE collection,
+		// so that pre-existing parameters bound elsewhere are excluded as rename candidates.
+		TSet<FName> ExistingBoundParams;
+		for (const TPair<FGameplayTag, FSubsonicEvent>& EventPair : Definition.GetEvents())
+		{
+			for (const FSubsonicEventActionDefinition& ActionDef : EventPair.Value.GetActionCollection())
+			{
+				if (const FSubsonicEventActionBase* ActionBase = ActionDef.GetAction().GetPtr())
+				{
+					for (const TPair<FName, FSubsonicEventActionBoundProperties>& BindingPair : ActionBase->GetBindings())
+					{
+						ExistingBoundParams.Add(BindingPair.Key);
+					}
+				}
+			}
+		}
+
+		// Scan a single bag for exactly one candidate parameter that:
+		// 1. Is NOT the old name
+		// 2. Is NOT already referenced by an existing binding
+		// 3. Is type-compatible with the stale action properties
+		// Returns NAME_None if zero or multiple candidates (ambiguous).
+		auto FindCandidateInBag = [&](const UPropertyBag* Bag)
+		{
+			const FName* BagCandidate = nullptr;
+				
+			if (Bag)
+			{
+				for (const FPropertyBagPropertyDesc& Desc : Bag->GetPropertyDescs())
+				{
+					if (Desc.Name != OldParamName && !ExistingBoundParams.Contains(Desc.Name) && Desc.CachedProperty)
+					{
+						const bool bAllCompatible = Algo::AllOf(StaleActionProps, [&Desc](const FProperty* ActionProp)
+						{
+							return ArePropertiesBindingCompatible(Desc.CachedProperty, ActionProp);
+						});
+
+						if (bAllCompatible)
+						{
+							if (!BagCandidate || *BagCandidate == Desc.Name)
+							{
+								BagCandidate = &Desc.Name;
+							}
+							else
+							{
+								// Ambiguous within this bag
+								BagCandidate = nullptr;
+								break;
+							}
+						}
+					}
+				}
+			}
+			return BagCandidate;
+		};
+
+		// Check collection bag first (most common rename target), then each event bag independently.
+		// Scanning per-bag prevents unrelated parameters in other bags from creating false ambiguity.
+		if (const FName* Result = FindCandidateInBag(Definition.GetParameters().GetPropertyBagStruct()))
+		{
+			return *Result;
+		}
+
+		for (const TPair<FGameplayTag, FSubsonicEvent>& EventPair : Definition.GetEvents())
+		{
+			if (const FName* Result = FindCandidateInBag(EventPair.Value.GetParameters().GetPropertyBagStruct()))
+			{
+				return *Result;
+			}
+		}
+
+		return NAME_None;
+	}
+#endif // WITH_EDITORONLY_DATA
+
+	void PromptForStaleBindings(const FSubsonicEventCollectionDefinition& Definition, const FStaleBindingsMap& StaleBindings, EStaleBindingResponse& OutResponse)
+	{
+#if WITH_EDITORONLY_DATA
+		// Detect rename candidates for each stale parameter.
+		TMap<FName, FName> RenameMap;
+		for (const TPair<FName, TArray<FStalePropertyHandle>>& Pair : StaleBindings)
+		{
+			const FName Candidate = FindRenameCandidate(Definition, Pair.Key, Pair.Value);
+			if (!Candidate.IsNone())
+			{
+				RenameMap.Add(Pair.Key, Candidate);
+			}
+		}
+
+		const bool bCanRebind = RenameMap.Num() == StaleBindings.Num();
+
+		if (bCanRebind)
+		{
+			// All stale params have rename candidates — offer the rebind option.
+			TArray<FText> RenameTexts;
+			for (const TPair<FName, FName>& Pair : RenameMap)
+			{
+				// \u2192 is the arrow character
+				RenameTexts.Add(FText::Format(LOCTEXT("RenameMapping", "{0} \u2192 {1}"),
+					FText::FromName(Pair.Key), FText::FromName(Pair.Value)));
+			}
+			const FText RenameList = FText::Join(FText::FromString(TEXT("\r\n\t* ")), RenameTexts);
+
+			const FText Title = LOCTEXT("StaleBindingRename_Title", "Parameter Renamed");
+			const FText Message = FText::Format(
+				LOCTEXT("StaleBindingRenameWarning",
+					"The following parameter(s) have been renamed:\r\n\t{0}\r\n"
+					"- Update affected action property bindings to the new name? (Yes)\r\n"
+					"- Remove the stale bindings? (No)\r\n"
+					"- Cancel the parameter rename? (Cancel)"),
+				RenameList);
+
+			const EAppReturnType::Type Result = FMessageDialog::Open(EAppMsgType::YesNoCancel, Message, Title);
+			switch (Result)
+			{
+				case EAppReturnType::Yes:
+				{
+					OutResponse = EStaleBindingResponse::Rebind;
+				}
+				break;
+
+				case EAppReturnType::No:
+				{
+					OutResponse = EStaleBindingResponse::RemoveBindings;
+				}
+				break;
+
+				case EAppReturnType::Cancel:
+				default:
+				{
+					OutResponse = EStaleBindingResponse::Revert;
+				}
+				break;
+			}
+		}
+		else
+		{
+			// No rename candidates — show the standard removal/revert dialog.
+			TArray<FText> ParamTexts;
+			Algo::Transform(StaleBindings, ParamTexts, [](const TPair<FName, TArray<FStalePropertyHandle>>& Pair) { return FText::FromName(Pair.Key); });
+			const FText ParamList = FText::Join(FText::FromString(TEXT("\r\n\t* ")), ParamTexts);
+
+			const FText Title = LOCTEXT("StaleBindingWarning_Title", "Binding Removal Required");
+			const FText Message = FText::Format(
+				LOCTEXT("StaleBindingWarning",
+					"The following parameter(s) have property bindings that are no longer valid "
+					"due to removal or type/name change:\r\n\t{0}\r\n\r\n"
+					"Would you like to remove the affected bindings (Yes) or revert the parameter change (No)?"),
+				ParamList);
+
+			const EAppReturnType::Type Result = FMessageDialog::Open(EAppMsgType::YesNo, Message, Title);
+			OutResponse = Result == EAppReturnType::No
+				? EStaleBindingResponse::Revert
+				: EStaleBindingResponse::RemoveBindings;
+		}
+#endif // WITH_EDITORONLY_DATA
+	}
+} // namespace UE::Subsonic::Core::BindingUtils
+
+#undef LOCTEXT_NAMESPACE
+#endif // WITH_EDITOR
diff --git a/Engine/Plugins/Experimental/Subsonic/Source/SubsonicCore/Private/SubsonicBuiltInParameters.cpp b/Engine/Plugins/Experimental/Subsonic/Source/SubsonicCore/Private/SubsonicBuiltInParameters.cpp
new file mode 100644
--- /dev/null
+++ b/Engine/Plugins/Experimental/Subsonic/Source/SubsonicCore/Private/SubsonicBuiltInParameters.cpp
@@ -0,0 +1,13 @@
+// Copyright Epic Games, Inc. All Rights Reserved.
+
+#include "SubsonicBuiltInParameters.h"
+
+namespace UE::Subsonic::BuiltInParameters
+{
+	const FName Volume         = "Volume";
+	const FName PitchShift     = "PitchShift";
+	const FName HighpassCutoff = "HighpassCutoff";
+	const FName LowpassCutoff  = "LowpassCutoff";
+	const FName FadeOutTime    = "FadeOutTime";
+
+} // namespace UE::Subsonic::BuiltInParameters
diff --git a/Engine/Plugins/Experimental/Subsonic/Source/SubsonicCore/Private/SubsonicCoreModule.cpp b/Engine/Plugins/Experimental/Subsonic/Source/SubsonicCore/Private/SubsonicCoreModule.cpp
new file mode 100644
--- /dev/null
+++ b/Engine/Plugins/Experimental/Subsonic/Source/SubsonicCore/Private/SubsonicCoreModule.cpp
@@ -0,0 +1,8 @@
+// Copyright Epic Games, Inc. All Rights Reserved.
+
+#include "Modules/ModuleManager.h"
+#include "SubsonicCoreLog.h"
+
+DEFINE_LOG_CATEGORY(LogSubsonic);
+
+IMPLEMENT_MODULE(FDefaultModuleImpl, SubsonicCore)
diff --git a/Engine/Plugins/Experimental/Subsonic/Source/SubsonicCore/Private/SubsonicEventCollection.cpp b/Engine/Plugins/Experimental/Subsonic/Source/SubsonicCore/Private/SubsonicEventCollection.cpp
new file mode 100644
--- /dev/null
+++ b/Engine/Plugins/Experimental/Subsonic/Source/SubsonicCore/Private/SubsonicEventCollection.cpp
@@ -0,0 +1,755 @@
+// Copyright Epic Games, Inc. All Rights Reserved.
+
+#include "SubsonicEventCollection.h"
+
+#include "ISubsonicEventRegistry.h"
+#include "NativeGameplayTags.h"
+#include "StructUtils/StructView.h"
+#include "SubsonicBindingUtils.h"
+#include "SubsonicCoreLog.h"
+#include "SubsonicExecutor.h"
+#include "SubsonicHandles.h"
+#include "UObject/UnrealType.h"
+
+
+namespace UE::Subsonic::Core
+{
+	UE_DEFINE_GAMEPLAY_TAG(TAG_SubsonicCore, "Subsonic");
+	UE_DEFINE_GAMEPLAY_TAG(TAG_SubsonicCore_Event_Play, "Subsonic.Event.Play");
+	UE_DEFINE_GAMEPLAY_TAG(TAG_SubsonicCore_Event_Stop, "Subsonic.Event.Stop");
+
+	namespace CollectionPrivate
+	{
+		constexpr bool bErrorOnMissingTag =
+#if WITH_EDITOR
+		// Validation will catch missing tags. While fixing up tags or
+		// creating new events with an unset tag, this messaging is too pervasive.
+		false;
+#else // !WITH_EDITOR
+		true;
+#endif // !WITH_EDITOR
+		constexpr uint32 InvalidId = static_cast<uint32>(INDEX_NONE);
+	} // namespace CollectionPrivate
+
+#if WITH_EDITOR
+	const TMap<FName, FSubsonicEventActionBoundProperties>& FSubsonicEventActionBase::GetBindings() const
+	{
+		return Bindings;
+	}
+
+	FName FSubsonicEventActionBase::GetBindingsParameterName()
+	{
+		return GET_MEMBER_NAME_CHECKED(FSubsonicEventActionBase, Bindings);
+	}
+
+	FText FSubsonicEventActionBase::GetDisplayInfo() const
+	{
+		return FText();
+	}
+
+	FName FSubsonicEventActionBase::GetBoundParameterForProperty(FName PropertyName) const
+	{
+#if WITH_EDITORONLY_DATA
+		for (const TPair<FName, FSubsonicEventActionBoundProperties>& Pair : Bindings)
+		{
+			if (Pair.Value.Properties.Contains(PropertyName))
+			{
+				return Pair.Key;
+			}
+		}
+#endif // WITH_EDITORONLY_DATA
+		return NAME_None;
+	}
+
+	void FSubsonicEventActionBase::AddPropertyBinding(FName ParameterName, FName PropertyName)
+	{
+#if WITH_EDITORONLY_DATA
+		RemovePropertyBinding(PropertyName);
+		Bindings.FindOrAdd(ParameterName).Properties.AddUnique(PropertyName);
+#endif // WITH_EDITORONLY_DATA
+	}
+
+	bool FSubsonicEventActionBase::RemoveAllPropertyBindingsToParameter(FName ParameterName)
+	{
+		return Bindings.Remove(ParameterName) > 0;
+	}
+
+	void FSubsonicEventActionBase::RemovePropertyBinding(FName PropertyName)
+	{
+#if WITH_EDITORONLY_DATA
+		TArray<FName> ToRemove;
+		for (TPair<FName, FSubsonicEventActionBoundProperties>& Pair : Bindings)
+		{
+			Pair.Value.Properties.Remove(PropertyName);
+			if (Pair.Value.Properties.IsEmpty())
+			{
+				ToRemove.Add(Pair.Key);
+			}
+		}
+
+		for (FName Name : ToRemove)
+		{
+			Bindings.Remove(Name);
+		}
+#endif // WITH_EDITORONLY_DATA
+	}
+#endif // WITH_EDITOR
+
+	const TInstancedStruct<FSubsonicEventActionBase>& FSubsonicEventActionDefinition::GetAction() const
+	{
+		return Action;
+	}
+
+#if WITH_EDITOR
+	void FSubsonicEventActionDefinition::SetAction(TInstancedStruct<FSubsonicEventActionBase> NewAction)
+	{
+		Action = MoveTemp(NewAction);
+	}
+
+	FSubsonicEventActionBase* FSubsonicEventActionDefinition::GetMutableActionBase()
+	{
+		return Action.GetMutablePtr();
+	}
+#endif // WITH_EDITOR
+
+#if WITH_EDITORONLY_DATA
+	const FInstancedPropertyBag& FSubsonicEvent::GetParameters() const
+	{
+		return Parameters;
+	}
+#endif // WITH_EDITORONLY_DATA
+
+#if WITH_EDITOR
+	void FSubsonicEvent::BindParameter(const FInstancedPropertyBag& InSourceParameters, FName ParamName)
+	{
+		const FInstancedPropertyBag* EffectiveBag = nullptr;
+		auto GetPropertyBagDesc = [&ParamName](const FInstancedPropertyBag& Bag)
+		{
+			const FPropertyBagPropertyDesc* Desc = nullptr;
+			if (const UPropertyBag* BagStruct = Bag.GetPropertyBagStruct())
+			{
+				Desc = BagStruct->FindPropertyDescByName(ParamName);
+			}
+
+			return Desc;
+		};
+
+		const FPropertyBagPropertyDesc* PropertyDesc = GetPropertyBagDesc(Parameters);
+		if (PropertyDesc)
+		{
+			EffectiveBag = &Parameters;
+		}
+		else
+		{
+			PropertyDesc = GetPropertyBagDesc(InSourceParameters);
+			EffectiveBag = &InSourceParameters;
+		}
+
+		if (!PropertyDesc || !PropertyDesc->CachedProperty)
+		{
+			return;
+		}
+
+		check(EffectiveBag);
+		const FConstStructView BagView = EffectiveBag->GetValue();
+
+		for (FSubsonicEventActionDefinition& ActionDefinition : ActionCollection)
+		{
+			const UScriptStruct* ActionStructType = ActionDefinition.Action.GetScriptStruct();
+			FSubsonicEventActionBase* ActionPtr = ActionDefinition.Action.GetMutablePtr();
+			if (!ActionPtr || !ActionStructType)
+			{
+				continue;
+			}
+
+			// Propagate parameter value to all type-compatible action properties.
+			// This does not record entries in the Bindings map — only explicit bindings
+			// created via AddPropertyBinding (the editor UI) are tracked there.
+			// Skip properties that are explicitly bound to a different parameter so that
+			// auto-propagation of one parameter cannot overwrite an explicit binding's value.
+			for (TFieldIterator<FProperty> PropIt(ActionStructType); PropIt; ++PropIt)
+			{
+				FProperty* ActionProp = *PropIt;
+				if (ActionProp->HasMetaData("NoBinding"))
+				{
+					continue;
+				}
+
+				// Require compatible property types (including inner types for containers) to safely copy the value.
+				if (!BindingUtils::ArePropertiesBindingCompatible(PropertyDesc->CachedProperty, ActionProp))
+				{
+					continue;
+				}
+
+				// Only propagate to properties explicitly bound to this parameter.
+				const FName BoundParam = ActionPtr->GetBoundParameterForProperty(ActionProp->GetFName());
+				if (BoundParam != ParamName)
+				{
+					continue;
+				}
+
+				if (const void* SrcPtr = PropertyDesc->CachedProperty->ContainerPtrToValuePtr<const void>(BagView.GetMemory()); ensure(SrcPtr))
+				{
+					void* DstPtr = ActionProp->ContainerPtrToValuePtr<void>(ActionPtr);
+					ActionProp->CopyCompleteValue(DstPtr, SrcPtr);
+				}
+			}
+		}
+	}
+
+	void FSubsonicEvent::UnbindParameter(FName ParamName)
+	{
+		for (FSubsonicEventActionDefinition& ActionDefinition : ActionCollection)
+		{
+			if (FSubsonicEventActionBase* ActionPtr = ActionDefinition.Action.GetMutablePtr())
+			{
+				ActionPtr->RemoveAllPropertyBindingsToParameter(ParamName);
+			}
+		}
+	}
+#endif // WITH_EDITOR
+
+	void FSubsonicEvent::Execute(const FSubsonicExecutor& InExecutor, const FEventHandle& InHandle) const
+	{
+		checkf(InExecutor.GetCollectionHandle().CollectionId == InHandle.Collection.CollectionId,
+			TEXT("Invalid execution context: executor CollectionId must match that of the given event handle"));
+
+		FActionHandle ActionHandle { .Event = InHandle };
+		for (int32 Index = 0; Index < ActionCollection.Num(); ++Index)
+		{
+			ActionHandle.Index = Index;
+
+			TConstStructView<FSubsonicEventActionBase> Action = ActionCollection[Index].Action;
+			if (const FSubsonicEventActionBase* ActionPtr = Action.GetPtr())
+			{
+				ActionPtr->Execute(InExecutor, ActionHandle);
+			}
+			else
+			{
+				UE_LOG(LogSubsonic, Warning, TEXT("Action not specified: Skipping action at '%s'"), *ActionHandle.ToString());
+			}
+		}
+	}
+
+	const TArray<FSubsonicEventActionDefinition>& FSubsonicEvent::GetActionCollection() const
+	{
+		return ActionCollection;
+	}
+
+#if WITH_EDITOR
+	FName FSubsonicEvent::GetActionCollectionPropertyName()
+	{
+		return GET_MEMBER_NAME_CHECKED(FSubsonicEvent, ActionCollection);
+	}
+
+	FName FSubsonicEvent::GetParametersPropertyName()
+	{
+		return GET_MEMBER_NAME_CHECKED(FSubsonicEvent, Parameters);
+	}
+#endif // WITH_EDITOR
+
+#if WITH_EDITOR
+	bool FSubsonicEvent::GetAutoAudition() const
+	{
+		return bAutoAudition;
+	}
+#endif // WITH_EDITOR
+
+	bool FSubsonicEvent::GetIsPublic() const
+	{
+		return bIsPublic;
+	}
+
+#if WITH_EDITOR
+	TArray<FSubsonicEventActionDefinition>& FSubsonicEvent::GetMutableActionCollection()
+	{
+		return ActionCollection;
+	}
+
+	void FSubsonicEvent::SetAutoAudition(bool bInAutoAudition)
+	{
+		bAutoAudition = bInAutoAudition;
+	}
+
+	void FSubsonicEvent::SetIsPublic(bool bInIsPublic)
+	{
+		bIsPublic = bInIsPublic;
+	}
+#endif // WITH_EDITOR
+
+	FSubsonicEventCollectionDefinition::FSubsonicEventCollectionDefinition(FSubsonicEventCollectionDefinition&& InOther)
+	{
+		if (this != &InOther)
+		{
+			Unregister();
+
+			CollectionId = InOther.CollectionId;
+			InOther.CollectionId = GetInvalidId();
+
+			DeviceId = InOther.DeviceId;
+			InOther.DeviceId = INDEX_NONE;
+
+			Events = MoveTemp(InOther.Events);
+			InOther.Events.Reset();
+
+#if WITH_EDITORONLY_DATA
+			Parameters = MoveTemp(InOther.Parameters);
+			InOther.Parameters.Reset();
+#endif // WITH_EDITORONLY_DATA
+		}
+	}
+
+	FSubsonicEventCollectionDefinition& FSubsonicEventCollectionDefinition::operator=(FSubsonicEventCollectionDefinition&& InOther)
+	{
+		if (&InOther != this)
+		{
+			Unregister();
+
+			CollectionId = InOther.CollectionId;
+			InOther.CollectionId = GetInvalidId();
+
+			DeviceId = InOther.DeviceId;
+			InOther.DeviceId = INDEX_NONE;
+
+			Events = MoveTemp(InOther.Events);
+			InOther.Events.Reset();
+
+#if WITH_EDITORONLY_DATA
+			Parameters = MoveTemp(InOther.Parameters);
+			InOther.Parameters.Reset();
+#endif // WITH_EDITORONLY_DATA
+		}
+
+		return *this;
+	}
+
+	void FSubsonicEventCollectionDefinition::AssignId()
+	{
+		static uint32 MaxId = CollectionPrivate::InvalidId;
+		CollectionId = ++MaxId;
+	}
+
+#if WITH_EDITOR
+	FSubsonicEventActionDefinition* FSubsonicEventCollectionDefinition::AddAction(const FEventHandle& InHandle, int32 InsertIndex)
+	{
+		if (FSubsonicEvent* Event = FindMutableEvent(InHandle))
+		{
+			TArray<FSubsonicEventActionDefinition>& ActionCollection = Event->GetMutableActionCollection();
+			if (!ActionCollection.IsValidIndex(InsertIndex))
+			{
+				return &ActionCollection.AddDefaulted_GetRef();
+			}
+
+			return &ActionCollection.InsertDefaulted_GetRef(InsertIndex);
+		}
+
+		return nullptr;
+	}
+
+	FSubsonicEvent* FSubsonicEventCollectionDefinition::AddEvent(const FGameplayTag& InTag, FSubsonicEvent NewEvent)
+	{
+		if (ContainsEvent(InTag))
+		{
+			return nullptr;
+		}
+
+		return &Events.Add(InTag, MoveTemp(NewEvent));
+	}
+
+	void FSubsonicEventCollectionDefinition::BindAllParameters()
+	{
+		if (const UPropertyBag* BagStruct = Parameters.GetPropertyBagStruct())
+		{
+			for (const FPropertyBagPropertyDesc& Desc : BagStruct->GetPropertyDescs())
+			{
+				BindParameter(Desc.Name);
+			}
+		}
+
+		// Bind event-level parameters that have no collection-level counterpart
+		for (TPair<FGameplayTag, FSubsonicEvent>& Pair : Events)
+		{
+			const UPropertyBag* EventBagStruct = Pair.Value.Parameters.GetPropertyBagStruct();
+			if (!EventBagStruct)
+			{
+				continue;
+			}
+			for (const FPropertyBagPropertyDesc& Desc : EventBagStruct->GetPropertyDescs())
+			{
+				const UPropertyBag* CollectionBagStruct = Parameters.GetPropertyBagStruct();
+				if (!CollectionBagStruct || !CollectionBagStruct->FindPropertyDescByName(Desc.Name))
+				{
+					Pair.Value.BindParameter(Parameters, Desc.Name);
+				}
+			}
+		}
+	}
+
+	void FSubsonicEventCollectionDefinition::BindParameter(FName ParamName)
+	{
+		for (TPair<FGameplayTag, FSubsonicEvent>& Pair : Events)
+		{
+			Pair.Value.BindParameter(Parameters, ParamName);
+		}
+	}
+
+	bool FSubsonicEventCollectionDefinition::RemoveAllPropertyBindingsToParameter(FName ParamName)
+	{
+		bool bBindingsRemoved = false;
+
+#if WITH_EDITORONLY_DATA
+		for (TPair<FGameplayTag, FSubsonicEvent>& EventPair : Events)
+		{
+			for (FSubsonicEventActionDefinition& ActionDef : EventPair.Value.GetMutableActionCollection())
+			{
+				if (FSubsonicEventActionBase* ActionBase = ActionDef.GetMutableActionBase())
+				{
+					bBindingsRemoved |= ActionBase->RemoveAllPropertyBindingsToParameter(ParamName);
+				}
+			}
+		}
+#endif // WITH_EDITORONLY_DATA
+
+		return bBindingsRemoved;
+	}
+
+	bool FSubsonicEventCollectionDefinition::RemoveStaleBindings(const FCollectionHandle& ParentHandle, bool bPromptRemoval)
+	{
+#if WITH_EDITORONLY_DATA
+		using EStaleBindingResponse = BindingUtils::EStaleBindingResponse;
+		using FStalePropertyHandle = BindingUtils::FStalePropertyHandle;
+
+		// Check for stale bindings caused by structural changes (type or name changes) to parameter bags.
+		const TMap<FName, TArray<FStalePropertyHandle>> StaleBindings = BindingUtils::FindStaleBindings(*this, ParentHandle);
+		if (!StaleBindings.IsEmpty())
+		{
+			if (bPromptRemoval)
+			{ 
+				EStaleBindingResponse Response = EStaleBindingResponse::RemoveBindings;
+				OnStaleBindingsDetected.Broadcast(StaleBindings, Response);
+				if (Response == EStaleBindingResponse::Revert)
+				{
+					// If an editor is open, it is expected to respond to the delegate above and undo the associated transaction in the
+					// next tick. A response of revert essentially signals the remaining removal to be skipped. (If no editor is open,
+					// stale bindings are removed and stale bindings are removed on next editor load).
+					return false;
+				}
+
+				if (Response == EStaleBindingResponse::Rebind)
+				{
+					// Update binding keys from old parameter name to the detected rename target.
+					// Falls back to removal for any parameter where no unique candidate is found.
+					for (const TPair<FName, TArray<FStalePropertyHandle>>& Pair : StaleBindings)
+					{
+						const FName NewParamName = BindingUtils::FindRenameCandidate(*this, Pair.Key, Pair.Value);
+						for (const FStalePropertyHandle& Entry : Pair.Value)
+						{
+							if (FSubsonicEventActionDefinition* Action = FindMutableAction(Entry.Action))
+							{
+								if (FSubsonicEventActionBase* ActionBase = Action->GetMutableActionBase())
+								{
+									ActionBase->RemovePropertyBinding(Entry.Property);
+									if (!NewParamName.IsNone())
+									{
+										ActionBase->AddPropertyBinding(NewParamName, Entry.Property);
+									}
+								}
+							}
+						}
+					}
+					return true;
+				}
+			}
+		}
+		else
+		{
+			return false;
+		}
+
+		for (const TPair<FName, TArray<FStalePropertyHandle>>& Pair : StaleBindings)
+		{
+			for (const FStalePropertyHandle& Entry : Pair.Value)
+			{
+				if (FSubsonicEventActionDefinition* Action = FindMutableAction(Entry.Action))
+				{
+					if (FSubsonicEventActionBase* ActionBase = Action->GetMutableActionBase())
+					{
+						ActionBase->RemovePropertyBinding(Entry.Property);
+					}
+				}
+			}
+		}
+
+		return true;
+#else // !WITH_EDITORONLY_DATA
+		return false;
+#endif // !WITH_EDITORONLY_DATA
+	}
+
+	void FSubsonicEventCollectionDefinition::UnbindParameter(FName ParamName)
+	{
+		for (TPair<FGameplayTag, FSubsonicEvent>& Pair : Events)
+		{
+			Pair.Value.UnbindParameter(ParamName);
+		}
+	}
+#endif // WITH_EDITOR
+
+	FSubsonicEventCollectionDefinition FSubsonicEventCollectionDefinition::Create(FName Name, TMap<FGameplayTag, FSubsonicEvent> InEvents, Audio::FDeviceId InDeviceId)
+	{
+		FSubsonicEventCollectionDefinition NewDefinition;
+		NewDefinition.AssignId();
+
+		const FCollectionHandle NewHandle
+		{
+	#if WITH_EDITORONLY_DATA
+			.CollectionName = Name,
+	#endif // WITH_EDITORONLY_DATA
+			.CollectionId = NewDefinition.GetCollectionId(),
+		};
+
+		NewDefinition.DeviceId = InDeviceId;
+		NewDefinition.Events = MoveTemp(InEvents);
+		NewDefinition.Register(NewHandle, InDeviceId);
+		return MoveTemp(NewDefinition);
+	}
+
+	FSubsonicEventCollectionDefinition::~FSubsonicEventCollectionDefinition()
+	{
+		Unregister();
+		CollectionId = GetInvalidId();
+		DeviceId = INDEX_NONE;
+	}
+
+#if WITH_EDITOR
+	bool FSubsonicEventCollectionDefinition::ClearActions(const FEventHandle& InHandle)
+	{
+		if (FSubsonicEvent* Event = FindMutableEvent(InHandle))
+		{
+			Event->GetMutableActionCollection().Empty();
+			return true;
+		}
+
+		return false;
+	}
+
+	void FSubsonicEventCollectionDefinition::ClearEvents()
+	{
+		Events.Empty();
+	}
+
+	bool FSubsonicEventCollectionDefinition::ContainsBinding(FName EventName, FName ParameterName) const
+	{
+		if (Parameters.FindPropertyDescByName(ParameterName) != nullptr)
+		{
+			return true;
+		}
+
+		if (const FSubsonicEvent* Event = FindEvent(EventName))
+		{
+			if (Event->GetParameters().FindPropertyDescByName(ParameterName) != nullptr)
+			{
+				return true;
+			}
+		}
+
+		return false;
+	}
+
+	BindingUtils::FOnStaleBindingsDetected& FSubsonicEventCollectionDefinition::GetOnStaleBindingsDetectedDelegate()
+	{
+		return OnStaleBindingsDetected;
+	}
+#endif // WITH_EDITOR
+
+	bool FSubsonicEventCollectionDefinition::ContainsEvent(const FGameplayTag& EventTag) const
+	{
+		return Events.Contains(EventTag);
+	}
+
+	bool FSubsonicEventCollectionDefinition::ContainsEvent(FName Name) const
+	{
+		const FGameplayTag EventTag = FGameplayTag::RequestGameplayTag(Name, CollectionPrivate::bErrorOnMissingTag);
+		if (EventTag.GetTagName().IsValid() || Name.IsNone())
+		{
+			return ContainsEvent(EventTag);
+		}
+
+		return false;
+
+	}
+
+	bool FSubsonicEventCollectionDefinition::ContainsEvent(const FEventHandle& InHandle) const
+	{
+		if (InHandle.Collection.CollectionId == CollectionId)
+		{
+			return ContainsEvent(InHandle.EventName);
+		}
+
+		return false;
+	}
+
+	const FSubsonicEventActionDefinition* FSubsonicEventCollectionDefinition::FindAction(const FActionHandle& InHandle) const
+	{
+		if (const FSubsonicEvent* Event = FindEvent(InHandle.Event))
+		{
+			const TArray<FSubsonicEventActionDefinition>& Actions = Event->GetActionCollection();
+			if (Actions.IsValidIndex(InHandle.Index))
+			{
+				return &Actions[InHandle.Index];
+			}
+		}
+
+		return nullptr;
+	}
+
+	const FSubsonicEvent* FSubsonicEventCollectionDefinition::FindEvent(const FGameplayTag& EventTag) const
+	{
+		return Events.Find(EventTag);
+	}
+
+	const FSubsonicEvent* FSubsonicEventCollectionDefinition::FindEvent(FName Name) const
+	{
+		const FGameplayTag EventTag = FGameplayTag::RequestGameplayTag(Name, CollectionPrivate::bErrorOnMissingTag);
+		return FindEvent(EventTag);
+	}
+
+	const FSubsonicEvent* FSubsonicEventCollectionDefinition::FindEvent(const FEventHandle& InHandle) const
+	{
+		if (InHandle.Collection.CollectionId == CollectionId)
+		{
+			return FindEvent(InHandle.EventName);
+		}
+
+		return nullptr;
+	}
+
+#if WITH_EDITOR
+	FSubsonicEventActionDefinition* FSubsonicEventCollectionDefinition::FindMutableAction(const FActionHandle& InHandle)
+	{
+		if (FSubsonicEvent* Event = FindMutableEvent(InHandle.Event))
+		{
+			TArray<FSubsonicEventActionDefinition>& Actions = Event->ActionCollection;
+			if (Actions.IsValidIndex(InHandle.Index))
+			{
+				return &Actions[InHandle.Index];
+			}
+		}
+
+		return nullptr;
+	}
+
+	FSubsonicEvent* FSubsonicEventCollectionDefinition::FindMutableEvent(const FEventHandle& InHandle)
+	{
+		if (InHandle.Collection.CollectionId == GetCollectionId())
+		{
+			const FGameplayTag EventTag = FGameplayTag::RequestGameplayTag(InHandle.EventName, CollectionPrivate::bErrorOnMissingTag);
+			if (EventTag.GetTagName().IsValid() || InHandle.EventName.IsNone())
+			{
+				return Events.Find(EventTag);
+			}
+		}
+
+		return nullptr;
+	}
+#endif // WITH_EDITOR
+
+	uint32 FSubsonicEventCollectionDefinition::GetCollectionId() const
+	{
+		return CollectionId;
+	}
+
+	const TMap<FGameplayTag, FSubsonicEvent>& FSubsonicEventCollectionDefinition::GetEvents() const
+	{
+		return Events;
+	}
+
+	uint32 FSubsonicEventCollectionDefinition::GetInvalidId()
+	{
+		return UE::Subsonic::Core::CollectionPrivate::InvalidId;
+	}
+
+#if WITH_EDITOR
+	FName FSubsonicEventCollectionDefinition::GetEventsPropertyName()
+	{
+#if WITH_EDITORONLY_DATA
+		return GET_MEMBER_NAME_CHECKED(FSubsonicEventCollectionDefinition, Events);
+#else
+		return { };
+#endif // !WITH_EDITORONLY_DATA
+	}
+
+	FName FSubsonicEventCollectionDefinition::GetParametersPropertyName()
+	{
+#if WITH_EDITORONLY_DATA
+		return GET_MEMBER_NAME_CHECKED(FSubsonicEventCollectionDefinition, Parameters);
+#else
+		return { };
+#endif // !WITH_EDITORONLY_DATA
+	}
+#endif // WITH_EDITOR
+
+#if WITH_EDITORONLY_DATA
+	const FInstancedPropertyBag& FSubsonicEventCollectionDefinition::GetParameters() const
+	{
+		return Parameters;
+	}
+#endif // WITH_EDITORONLY_DATA
+
+	bool FSubsonicEventCollectionDefinition::IsValid() const
+	{
+		return CollectionId != INDEX_NONE;
+	}
+
+	void FSubsonicEventCollectionDefinition::Register(const Core::FCollectionHandle& InHandle, Audio::FDeviceId InDeviceId)
+	{
+		if (IsValid())
+		{
+			DeviceId = InDeviceId;
+			ISubsonicEventRegistry::GetChecked().OnCollectionRegistered(InHandle, DeviceId);
+		}
+	}
+
+#if WITH_EDITOR
+	bool FSubsonicEventCollectionDefinition::RemoveAction(const FActionHandle& InHandle)
+	{
+		if (FSubsonicEvent* Event = FindMutableEvent(InHandle.Event))
+		{
+			TArray<FSubsonicEventActionDefinition>& Actions = Event->GetMutableActionCollection();
+			if (Actions.IsValidIndex(InHandle.Index))
+			{
+				Actions.RemoveAt(InHandle.Index);
+				return true;
+			}
+		}
+
+		return false;
+	}
+
+	bool FSubsonicEventCollectionDefinition::RemoveEvent(const FEventHandle& InHandle)
+	{
+		if (InHandle.Collection.CollectionId == CollectionId)
+		{
+			const FGameplayTag EventTag = FGameplayTag::RequestGameplayTag(InHandle.EventName, CollectionPrivate::bErrorOnMissingTag);
+			if (EventTag.GetTagName().IsValid() || InHandle.EventName.IsNone())
+			{
+				return Events.Remove(EventTag) > 0;
+			}
+		}
+
+		return false;
+	}
+#endif // WITH_EDITOR
+
+	void FSubsonicEventCollectionDefinition::Unregister()
+	{
+		if (ISubsonicEventRegistry* EventRegistry = ISubsonicEventRegistry::Get(); EventRegistry && IsValid())
+		{
+			EventRegistry->OnCollectionUnregistered(FCollectionHandle { .CollectionId = CollectionId }, DeviceId);
+		}
+
+		CollectionId = CollectionPrivate::InvalidId;
+		DeviceId = INDEX_NONE;
+	}
+
+} // namespace UE::Subsonic::Core
diff --git a/Engine/Plugins/Experimental/Subsonic/Source/SubsonicCore/Private/SubsonicEventSubscriberInterface.cpp b/Engine/Plugins/Experimental/Subsonic/Source/SubsonicCore/Private/SubsonicEventSubscriberInterface.cpp
new file mode 100644
--- /dev/null
+++ b/Engine/Plugins/Experimental/Subsonic/Source/SubsonicCore/Private/SubsonicEventSubscriberInterface.cpp
@@ -0,0 +1,48 @@
+// Copyright Epic Games, Inc. All Rights Reserved.
+
+#include "SubsonicEventSubscriberInterface.h"
+
+#include "ISubsonicEventRegistry.h"
+#include "SubsonicExecutor.h"
+#include "SubsonicHandles.h"
+
+
+namespace UE::Subsonic::Core
+{
+	void ISubsonicEventSubscriberInterface::OnCollectionRegistered(const FCollectionHandle& InCollection)
+	{
+	}
+
+	void ISubsonicEventSubscriberInterface::OnCollectionUnregistered(const FCollectionHandle& InCollection)
+	{
+	}
+
+	void ISubsonicEventSubscriberInterface::OnEventPreExecute(const FSubsonicExecutor& InExecutor, const FEventHandle& InHandle)
+	{
+	}
+
+	void ISubsonicEventSubscriberInterface::OnEventPostExecute(const FSubsonicExecutor& InExecutor, const FEventHandle& InHandle)
+	{
+	}
+
+	void ISubsonicEventSubscriberInterface::OnExecutorRegistered(const FSubsonicExecutor& InExecutor)
+	{
+	}
+
+	void ISubsonicEventSubscriberInterface::OnExecutorUnregistered(const FSubsonicExecutor& InExecutor)
+	{
+	}
+
+	void ISubsonicEventSubscriberInterface::Register()
+	{
+		ISubsonicEventRegistry::GetChecked().OnEventSubscriberRegistered(*this);
+	}
+
+	void ISubsonicEventSubscriberInterface::Unregister()
+	{
+		if (ISubsonicEventRegistry* Registry = ISubsonicEventRegistry::Get())
+		{
+			Registry->OnEventSubscriberUnregistered(*this);
+		}
+	}
+} // namespace UE::Subsonic::Core
\ No newline at end of file
diff --git a/Engine/Plugins/Experimental/Subsonic/Source/SubsonicCore/Private/SubsonicExecutor.cpp b/Engine/Plugins/Experimental/Subsonic/Source/SubsonicCore/Private/SubsonicExecutor.cpp
new file mode 100644
--- /dev/null
+++ b/Engine/Plugins/Experimental/Subsonic/Source/SubsonicCore/Private/SubsonicExecutor.cpp
@@ -0,0 +1,176 @@
+// Copyright Epic Games, Inc. All Rights Reserved.
+
+#include "SubsonicExecutor.h"
+
+#include "ISubsonicEventRegistry.h"
+#include "SubsonicCoreLog.h"
+#include "SubsonicEventCollection.h"
+#include "SubsonicHandles.h"
+
+
+namespace UE::Subsonic::Core
+{
+	namespace ExecutorPrivate
+	{
+		constexpr uint32 InvalidId = static_cast<uint32>(INDEX_NONE);
+		Audio::FDeviceId InvalidDeviceId = static_cast<Audio::FDeviceId>(INDEX_NONE);
+	} // namespace ExecutorPrivate
+
+	FSubsonicExecutor::FSubsonicExecutor()
+		: DeviceId(ExecutorPrivate::InvalidDeviceId)
+	{
+	}
+
+	FSubsonicExecutor::FSubsonicExecutor(FSubsonicExecutor&& InOther)
+	{
+		Unregister();
+
+		CollectionAccessor = MoveTemp(InOther.CollectionAccessor);
+		Parameters = MoveTemp(InOther.Parameters);
+
+		DeviceId = InOther.DeviceId;
+		InOther.DeviceId = ExecutorPrivate::InvalidDeviceId;
+
+		ExecutorId = InOther.ExecutorId;
+		InOther.ExecutorId = ExecutorPrivate::InvalidId;
+	}
+
+	FSubsonicExecutor::~FSubsonicExecutor()
+	{
+		Unregister();
+	}
+
+	FSubsonicExecutor& FSubsonicExecutor::operator=(FSubsonicExecutor&& InOther)
+	{
+		if (&InOther != this)
+		{
+			Unregister();
+
+			CollectionAccessor = MoveTemp(InOther.CollectionAccessor);
+			Parameters = MoveTemp(InOther.Parameters);
+
+			DeviceId = InOther.DeviceId;
+			InOther.DeviceId = ExecutorPrivate::InvalidDeviceId;
+
+			ExecutorId = InOther.ExecutorId;
+			InOther.ExecutorId = ExecutorPrivate::InvalidId;
+		}
+
+		return *this;
+	}
+
+	FSubsonicExecutor FSubsonicExecutor::Create(Audio::FDeviceId InDeviceId, FSubsonicExecutor::FCollectionAccessorPtr&& CollectionPtr)
+	{
+		static uint32 MaxExecutorId = ExecutorPrivate::InvalidId;
+		FSubsonicExecutor NewHandle;
+		NewHandle.CollectionAccessor = MoveTemp(CollectionPtr);
+		NewHandle.ExecutorId = ++MaxExecutorId;
+		NewHandle.DeviceId = InDeviceId;
+		NewHandle.Register();
+		return NewHandle;
+	}
+
+	FCollectionHandle FSubsonicExecutor::GetCollectionHandle() const
+	{
+		if (CollectionAccessor.IsValid())
+		{
+			return CollectionAccessor->GetHandle();
+		}
+
+		return FCollectionHandle();
+	}
+
+	Audio::FDeviceId FSubsonicExecutor::GetDeviceId() const
+	{
+		return DeviceId;
+	}
+
+	uint32 FSubsonicExecutor::GetId() const
+	{
+		return ExecutorId;
+	}
+
+	bool FSubsonicExecutor::ExecuteEvent(FName EventName) const
+	{
+		if (CollectionAccessor.IsValid())
+		{
+			if (const FSubsonicEventCollectionDefinition* Definition = CollectionAccessor->GetDefinition())
+			{
+				if (const FSubsonicEvent* Event = Definition->FindEvent(EventName))
+				{
+					if (Event->GetIsPublic())
+					{
+						const FCollectionHandle& CollectionHandle = CollectionAccessor->GetHandle();
+						const FEventHandle EventHandle{ .Collection = CollectionHandle, .EventName = EventName };
+
+						ISubsonicEventRegistry& Registry = ISubsonicEventRegistry::GetChecked();
+						Registry.OnEventPreExecute(*this, EventHandle);
+						Event->Execute(*this, EventHandle);
+						Registry.OnEventPostExecute(*this, EventHandle);
+						return true;
+					}
+				}
+
+#if !NO_LOGGING
+				const FString CollectionString =
+#if WITH_EDITORONLY_DATA
+					FString::Printf(TEXT(" in collection '%s'"), *CollectionAccessor->GetHandle().CollectionName.ToString());
+#else // !WITH_EDITORONLY_DATA
+					TEXT("");
+#endif // !WITH_EDITORONLY_DATA
+
+				UE_LOG(LogSubsonic, Warning, TEXT("Failed to execute event '%s'%s: Event not found or accessible (not public)."), *EventName.ToString(), *CollectionString);
+#endif // !NO_LOGGING
+			}
+		}
+
+		return false;
+	}
+
+	uint32 FSubsonicExecutor::GetInvalidId()
+	{
+		return ExecutorPrivate::InvalidId;
+	}
+
+	FSubsonicExecutor FSubsonicExecutor::GetInvalid()
+	{
+		return { };
+	}
+
+	bool FSubsonicExecutor::IsValid() const
+	{
+		return ExecutorId != ExecutorPrivate::InvalidId;
+	}
+
+	void FSubsonicExecutor::Register()
+	{
+		ISubsonicEventRegistry::GetChecked().OnExecutorRegistered(*this);
+	}
+
+	void FSubsonicExecutor::Unregister()
+	{
+		if (IsValid())
+		{
+			ISubsonicEventRegistry::GetChecked().OnExecutorUnregistered(*this);
+		}
+
+		CollectionAccessor.Reset();
+		Parameters.Reset();
+		ExecutorId = ExecutorPrivate::InvalidId;
+	}
+
+	void FSubsonicExecutor::SetParameters(FSubsonicParameterStore InParams)
+	{
+		Parameters = MoveTemp(InParams);
+	}
+
+	const FSubsonicParameterStore& FSubsonicExecutor::GetParameters() const
+	{
+		return Parameters;
+	}
+
+	FString FSubsonicExecutor::ToString() const
+	{
+		return FString::Printf(TEXT("Id: %u, DeviceId: %u"), ExecutorId, DeviceId);
+	}
+} // namespace UE::Subsonic::Core
diff --git a/Engine/Plugins/Experimental/Subsonic/Source/SubsonicCore/Private/SubsonicHandles.cpp b/Engine/Plugins/Experimental/Subsonic/Source/SubsonicCore/Private/SubsonicHandles.cpp
new file mode 100644
--- /dev/null
+++ b/Engine/Plugins/Experimental/Subsonic/Source/SubsonicCore/Private/SubsonicHandles.cpp
@@ -0,0 +1,36 @@
+// Copyright Epic Games, Inc. All Rights Reserved.
+
+#include "SubsonicHandles.h"
+
+
+namespace UE::Subsonic::Core
+{
+	namespace HandlesPrivate
+	{
+		uint32 InvalidHandleId = static_cast<uint32>(INDEX_NONE);
+	} // namespace HandlesPrivate
+
+	uint32 FCollectionHandle::GetInvalidId()
+	{
+		return HandlesPrivate::InvalidHandleId;
+	}
+
+	FString FCollectionHandle::ToString() const
+	{
+#if WITH_EDITORONLY_DATA
+		return FString::Printf(TEXT("%s (Id: %u)"), *CollectionName.ToString(), CollectionId);
+#else // !WITH_EDITORONLY_DATA
+		return FString::Printf(TEXT("Id: %u"), CollectionId);
+#endif // !WITH_EDITORONLY_DATA
+	}
+
+	FString FEventHandle::ToString() const
+	{
+		return FString::Printf(TEXT("%s/%s"), *Collection.ToString(), *EventName.ToString());
+	}
+
+	FString FActionHandle::ToString() const
+	{
+		return FString::Printf(TEXT("%s[%i]"), *Event.ToString(), Index);
+	}
+} // namespace UE::Subsonic::Core
\ No newline at end of file
diff --git a/Engine/Plugins/Experimental/Subsonic/Source/SubsonicCore/Private/SubsonicParameterStore.cpp b/Engine/Plugins/Experimental/Subsonic/Source/SubsonicCore/Private/SubsonicParameterStore.cpp
new file mode 100644
--- /dev/null
+++ b/Engine/Plugins/Experimental/Subsonic/Source/SubsonicCore/Private/SubsonicParameterStore.cpp
@@ -0,0 +1,47 @@
+// Copyright Epic Games, Inc. All Rights Reserved.
+
+#include "SubsonicParameterStore.h"
+
+
+namespace UE::Subsonic
+{
+	bool FSubsonicParameterStore::HasParameter(FName Name) const
+	{
+		const UPropertyBag* BagStruct = Bag.GetPropertyBagStruct();
+		return BagStruct != nullptr && BagStruct->FindPropertyDescByName(Name) != nullptr;
+	}
+
+	void FSubsonicParameterStore::RemoveParameter(FName Name)
+	{
+		Bag.RemovePropertyByName(Name);
+	}
+
+	void FSubsonicParameterStore::Reset()
+	{
+		Bag = FInstancedPropertyBag{};
+	}
+
+	bool FSubsonicParameterStore::IsEmpty() const
+	{
+		return Bag.GetNumPropertiesInBag() == 0;
+	}
+
+	void FSubsonicParameterStore::MergeFrom(const FSubsonicParameterStore& Other)
+	{
+		const UPropertyBag* OtherBagStruct = Other.Bag.GetPropertyBagStruct();
+		if (!OtherBagStruct)
+		{
+			return;
+		}
+
+		const TConstArrayView<FPropertyBagPropertyDesc> OtherDescs = OtherBagStruct->GetPropertyDescs();
+		if (OtherDescs.IsEmpty())
+		{
+			return;
+		}
+
+		// Ensure all properties from Other exist in this bag, then overlay their values.
+		Bag.AddProperties(OtherDescs);
+		Bag.CopyMatchingValuesByName(Other.Bag);
+	}
+} // namespace UE::Subsonic
diff --git a/Engine/Plugins/Experimental/Subsonic/Source/SubsonicCore/Public/ISubsonicEventRegistry.h b/Engine/Plugins/Experimental/Subsonic/Source/SubsonicCore/Public/ISubsonicEventRegistry.h
new file mode 100644
--- /dev/null
+++ b/Engine/Plugins/Experimental/Subsonic/Source/SubsonicCore/Public/ISubsonicEventRegistry.h
@@ -0,0 +1,49 @@
+// Copyright Epic Games, Inc. All Rights Reserved.
+#pragma once
+
+#include "AudioDefines.h"
+#include "Templates/UniquePtr.h"
+
+#define UE_API SUBSONICCORE_API
+
+
+namespace UE::Subsonic::Core
+{
+	// Forward Declarations
+	struct FCollectionHandle;
+	struct FSubsonicExecutor;
+	struct FEventHandle;
+
+	class ISubsonicEventSubscriberInterface;
+
+
+	// Registry interface for events and their associated collections/executors.
+	class ISubsonicEventRegistry
+	{
+	public:
+		virtual ~ISubsonicEventRegistry() = default;
+
+		UE_API static ISubsonicEventRegistry* Get();
+		UE_API static ISubsonicEventRegistry& GetChecked();
+
+		UE_API static void Initialize(TUniquePtr<ISubsonicEventRegistry>&& InRegistry);
+		UE_API static void Deinitialize();
+
+		virtual void OnCollectionRegistered(const FCollectionHandle& InHandle, Audio::FDeviceId DeviceId) = 0;
+		virtual void OnCollectionUnregistered(const FCollectionHandle& InHandle, Audio::FDeviceId DeviceId) = 0;
+
+		virtual void OnExecutorRegistered(const FSubsonicExecutor& InExecutor) = 0;
+		virtual void OnExecutorUnregistered(const FSubsonicExecutor& InExecutor) = 0;
+
+		virtual void OnEventPreExecute(const FSubsonicExecutor& InExecutor, const FEventHandle& InHandle) = 0;
+		virtual void OnEventPostExecute(const FSubsonicExecutor& InExecutor, const FEventHandle& InHandle) = 0;
+
+		virtual void OnEventSubscriberRegistered(ISubsonicEventSubscriberInterface& EventSubscriber) = 0;
+		virtual void OnEventSubscriberUnregistered(ISubsonicEventSubscriberInterface& EventSubscriber) = 0;
+
+	private:
+		static TUniquePtr<ISubsonicEventRegistry> Instance;
+	};
+} // namespace UE::Subsonic::Core
+
+#undef UE_API // SUBSONICCORE_API
diff --git a/Engine/Plugins/Experimental/Subsonic/Source/SubsonicCore/Public/SubsonicBindingUtils.h b/Engine/Plugins/Experimental/Subsonic/Source/SubsonicCore/Public/SubsonicBindingUtils.h
new file mode 100644
--- /dev/null
+++ b/Engine/Plugins/Experimental/Subsonic/Source/SubsonicCore/Public/SubsonicBindingUtils.h
@@ -0,0 +1,66 @@
+// Copyright Epic Games, Inc. All Rights Reserved.
+#pragma once
+
+#include "Containers/Map.h"
+#include "Delegates/Delegate.h"
+#include "Delegates/DelegateBase.h"
+#include "Delegates/DelegateCombinations.h"
+#include "SubsonicHandles.h"
+
+
+#define UE_API SUBSONICCORE_API
+
+namespace UE::Subsonic::Core
+{
+	namespace BindingUtils
+	{
+#if WITH_EDITOR
+		// Response returned by OnStaleBindingsDetected to indicate whether to remove all stale
+		// bindings (proceeding with the parameter structural change), revert the change, or rebind
+		// stale bindings to a renamed parameter (where applicable).
+		enum class EStaleBindingResponse : uint8
+		{
+			RemoveBindings,
+			Revert,
+			Rebind
+		};
+
+		// Represents minimal information required to locate a stale binding entry.
+		struct FStalePropertyHandle
+		{
+			// Action containing stale binding.
+			FActionHandle Action;
+
+			// Name of property on action that is stale
+			FName Property;
+		};
+
+		// Map of stale parameter name to action pointing to stale property
+		using FStaleBindingsMap = TMap<FName, TArray<FStalePropertyHandle>>;
+
+		// Fired when a structural parameter change (type or name change) creates stale editor bindings.
+		DECLARE_MULTICAST_DELEGATE_TwoParams(FOnStaleBindingsDetected, const FStaleBindingsMap&, EStaleBindingResponse&);
+
+		// Checks whether a property bag value can be safely copied into a destination (action) property.
+		// Uses SameType for exact matching, with a fallback for object properties to allow subclass
+		// assignment. Does not currently handle inner-type variance for container properties (TArray, etc.).
+		UE_INTERNAL UE_API bool ArePropertiesBindingCompatible(const FProperty* FromBag, const FProperty* ToAction);
+
+		// For a stale parameter name, finds a likely rename candidate in the current bags: a parameter
+		// with a compatible type that is not already referenced by any existing binding on the affected
+		// actions. Returns NAME_None if no unique candidate is found.
+		UE_INTERNAL UE_API FName FindRenameCandidate(const FSubsonicEventCollectionDefinition& Definition, FName OldParamName, const TArray<FStalePropertyHandle>& StaleEntries);
+
+		// Returns map of parameter names to stale binding entries for bindings that are no longer valid
+		// (parameter removed or type no longer matches). Checks event-level bags first, then collection-level.
+		UE_INTERNAL UE_API TMap<FName, TArray<FStalePropertyHandle>> FindStaleBindings(const FSubsonicEventCollectionDefinition& Definition, const FCollectionHandle& ParentHandle);
+
+		// Prompts the user to handle stale bindings caused by a parameter rename or removal.
+		// Detects rename candidates and shows either a 3-option (rename/remove/revert) or 2-option
+		// (remove/revert) dialog. Sets OutResponse accordingly.
+		UE_INTERNAL UE_API void PromptForStaleBindings(const FSubsonicEventCollectionDefinition& Definition, const FStaleBindingsMap& StaleBindings, EStaleBindingResponse& OutResponse);
+#endif // WITH_EDITOR
+	} // namespace BindingUtils
+} // namespace UE::Subsonic::Core
+
+#undef UE_API // SUBSONICCORE_API
diff --git a/Engine/Plugins/Experimental/Subsonic/Source/SubsonicCore/Public/SubsonicBuiltInParameters.h b/Engine/Plugins/Experimental/Subsonic/Source/SubsonicCore/Public/SubsonicBuiltInParameters.h
new file mode 100644
--- /dev/null
+++ b/Engine/Plugins/Experimental/Subsonic/Source/SubsonicCore/Public/SubsonicBuiltInParameters.h
@@ -0,0 +1,15 @@
+// Copyright Epic Games, Inc. All Rights Reserved.
+
+#pragma once
+
+#include "UObject/NameTypes.h"
+
+namespace UE::Subsonic::BuiltInParameters
+{
+	extern SUBSONICCORE_API const FName Volume;
+	extern SUBSONICCORE_API const FName PitchShift;
+	extern SUBSONICCORE_API const FName HighpassCutoff;
+	extern SUBSONICCORE_API const FName LowpassCutoff;
+	extern SUBSONICCORE_API const FName FadeOutTime;
+
+} // namespace UE::Subsonic::BuiltInParameters
diff --git a/Engine/Plugins/Experimental/Subsonic/Source/SubsonicCore/Public/SubsonicCoreLog.h b/Engine/Plugins/Experimental/Subsonic/Source/SubsonicCore/Public/SubsonicCoreLog.h
new file mode 100644
--- /dev/null
+++ b/Engine/Plugins/Experimental/Subsonic/Source/SubsonicCore/Public/SubsonicCoreLog.h
@@ -0,0 +1,6 @@
+// Copyright Epic Games, Inc. All Rights Reserved.
+#pragma once
+
+#include "Logging/LogMacros.h"
+
+SUBSONICCORE_API DECLARE_LOG_CATEGORY_EXTERN(LogSubsonic, Log, All);
diff --git a/Engine/Plugins/Experimental/Subsonic/Source/SubsonicCore/Public/SubsonicEventCollection.h b/Engine/Plugins/Experimental/Subsonic/Source/SubsonicCore/Public/SubsonicEventCollection.h
new file mode 100644
--- /dev/null
+++ b/Engine/Plugins/Experimental/Subsonic/Source/SubsonicCore/Public/SubsonicEventCollection.h
@@ -0,0 +1,317 @@
+// Copyright Epic Games, Inc. All Rights Reserved.
+#pragma once
+
+#include "AudioDefines.h"
+#include "GameplayTagContainer.h"
+#include "NativeGameplayTags.h"
+#include "StructUtils/InstancedStruct.h"
+#include "StructUtils/PropertyBag.h"
+#include "SubsonicBindingUtils.h"
+#include "SubsonicHandles.h"
+#include "UObject/NameTypes.h"
+
+#include "SubsonicEventCollection.generated.h"
+
+// Forward Declarations
+class FProperty;
+
+#define UE_API SUBSONICCORE_API
+
+namespace UE::Subsonic::Core
+{
+	// Forward Declarations
+	struct FSubsonicExecutor;
+
+	UE_API UE_DECLARE_GAMEPLAY_TAG_EXTERN(TAG_SubsonicCore)
+	UE_API UE_DECLARE_GAMEPLAY_TAG_EXTERN(TAG_SubsonicCore_Event_Play)
+	UE_API UE_DECLARE_GAMEPLAY_TAG_EXTERN(TAG_SubsonicCore_Event_Stop)
+
+	USTRUCT()
+	struct FSubsonicEventActionBoundProperties
+	{
+		GENERATED_BODY()
+
+		UPROPERTY()
+		TArray<FName> Properties;
+	};
+
+	USTRUCT(MinimalAPI, BlueprintType, meta = (Abstract, Hidden))
+	struct FSubsonicEventActionBase
+	{
+		GENERATED_BODY()
+
+		virtual ~FSubsonicEventActionBase() = default;
+
+#if WITH_EDITOR
+		UE_INTERNAL UE_API const TMap<FName, FSubsonicEventActionBoundProperties>& GetBindings() const;
+		UE_INTERNAL UE_API static FName GetBindingsParameterName();
+
+		// Returns display info that is custom to the given action instance to be shown
+		// adjacent to the action, for example, in the Subsonic Event Editor action tree.
+		UE_API virtual FText GetDisplayInfo() const;
+#endif // WITH_EDITOR
+
+	public:
+#if WITH_EDITOR
+		// Returns the parameter name that the given action property is currently bound to, or NAME_None if unbound.
+		UE_API FName GetBoundParameterForProperty(FName PropertyName) const;
+
+		// Binds an action property to a named parameter, replacing any previous binding for that property.
+		UE_API void AddPropertyBinding(FName ParameterName, FName PropertyName);
+
+		// Removes all property bindings associated with the parameter of the given name.
+		// Returns true if binding was found and removed, false if not.
+		UE_API bool RemoveAllPropertyBindingsToParameter(FName ParameterName);
+
+		// Removes any binding for the given action property.
+		UE_API void RemovePropertyBinding(FName PropertyName);
+#endif // WITH_EDITOR
+
+	protected:
+		virtual void Execute(const FSubsonicExecutor& InExecutor, const FActionHandle& InHandle) const { }
+
+#if WITH_EDITORONLY_DATA
+		UPROPERTY()
+		TMap<FName, FSubsonicEventActionBoundProperties> Bindings;
+#endif // WITH_EDITORONLY_DATA
+
+		friend struct FSubsonicEvent;
+	};
+
+	// Action definition iteratively executed per Subsonic event.
+	USTRUCT(MinimalAPI, BlueprintType)
+	struct FSubsonicEventActionDefinition
+	{
+		GENERATED_BODY()
+
+		UE_API const TInstancedStruct<FSubsonicEventActionBase>& GetAction() const;
+
+#if WITH_EDITOR
+		UE_API void SetAction(TInstancedStruct<FSubsonicEventActionBase> NewAction);
+
+		// Returns a mutable pointer to the underlying action base, or nullptr if no action is set.
+		UE_INTERNAL UE_API FSubsonicEventActionBase* GetMutableActionBase();
+#endif // WITH_EDITOR
+ 
+	private:
+		friend struct FSubsonicEvent;
+
+		UPROPERTY(EditAnywhere, BlueprintReadOnly, Category = "Actions", meta = (AllowPrivateAccess))
+		TInstancedStruct<FSubsonicEventActionBase> Action;
+	};
+
+	USTRUCT(MinimalAPI, BlueprintType)
+	struct FSubsonicEvent
+	{
+		GENERATED_BODY()
+
+		virtual ~FSubsonicEvent() = default;
+
+#if WITH_EDITOR
+		UE_INTERNAL UE_API static FName GetActionCollectionPropertyName();
+		UE_INTERNAL UE_API static FName GetParametersPropertyName();
+#endif // WITH_EDITOR
+
+		UE_API const TArray<FSubsonicEventActionDefinition>& GetActionCollection() const;
+
+#if WITH_EDITORONLY_DATA
+		UE_API const FInstancedPropertyBag& GetParameters() const;
+#endif // WITH_EDITORONLY_DATA
+
+#if WITH_EDITOR
+		UE_API bool GetAutoAudition() const;
+#endif // WITH_EDITOR
+
+		UE_API bool GetIsPublic() const;
+
+#if WITH_EDITOR
+		UE_INTERNAL UE_API TArray<FSubsonicEventActionDefinition>& GetMutableActionCollection();
+		UE_INTERNAL UE_API void SetAutoAudition(bool bInAutoAudition);
+		UE_INTERNAL UE_API void SetIsPublic(bool bInIsPublic);
+#endif // WITH_EDITOR
+
+	private:
+		friend struct FSubsonicEventCollectionDefinition;
+		friend struct FSubsonicExecutor;
+
+#if WITH_EDITOR
+		UE_INTERNAL void BindParameter(const FInstancedPropertyBag& InSourceParameters, FName Name);
+		UE_INTERNAL void UnbindParameter(FName Name);
+#endif // WITH_EDITOR
+
+	private:
+		// Called by executor, iterates through all actions on the given event definition.
+		// Asserts if the provided executor and handle have mismatched collection contexts.
+		UE_INTERNAL void Execute(const FSubsonicExecutor& InExecutor, const FEventHandle& InHandle) const;
+
+		UPROPERTY(EditAnywhere, BlueprintReadOnly, Category = "Event", meta = (AllowPrivateAccess))
+		bool bIsPublic = true;
+
+#if WITH_EDITORONLY_DATA
+		// If true, execute this event as soon as an audition executor is created.
+		UPROPERTY(EditAnywhere, Transient, Category = "Event", meta = (EditCondition = bIsPublic))
+		bool bAutoAudition = false;
+
+		UPROPERTY(EditAnywhere, Category = "Event")
+		FInstancedPropertyBag Parameters;
+#endif // WITH_EDITORONLY_DATA
+
+		UPROPERTY(EditAnywhere, BlueprintReadOnly, Category = "Actions", meta = (AllowPrivateAccess))
+		TArray<FSubsonicEventActionDefinition> ActionCollection;
+	};
+
+
+	USTRUCT(MinimalAPI, BlueprintType)
+	struct FSubsonicEventCollectionDefinition final
+	{
+		GENERATED_BODY()
+
+		// Generates empty, invalid event collection. Required to be public
+		// to support UObject reflection. To produce a valid definition natively
+		// without UObject support, call static Create function below.
+		FSubsonicEventCollectionDefinition() = default;
+
+		FSubsonicEventCollectionDefinition(FSubsonicEventCollectionDefinition&& Definition);
+		FSubsonicEventCollectionDefinition(const FSubsonicEventCollectionDefinition& Definition) = delete;
+
+		UE_API FSubsonicEventCollectionDefinition& operator=(FSubsonicEventCollectionDefinition&& Definition);
+		FSubsonicEventCollectionDefinition& operator=(const FSubsonicEventCollectionDefinition& Definition) = delete;
+
+		UE_API ~FSubsonicEventCollectionDefinition();
+
+		// Create a collection definition with a valid transient Id for tracking
+		// within the Subsonic System. Used to create and register definitions
+		// outside of UObject reflection.
+		static FSubsonicEventCollectionDefinition Create(FName Name, TMap<FGameplayTag, FSubsonicEvent> InEvents, Audio::FDeviceId DeviceId);
+
+#if WITH_EDITOR
+		// Adds action to event with handle if it exists. If insert index set to less than 0 (default), adds to end of collection array.
+		UE_INTERNAL UE_API FSubsonicEventActionDefinition* AddAction(const FEventHandle& InHandle, int32 InsertIndex = INDEX_NONE);
+
+		// Adds event with given tag if does not yet exists. Returns pointer to new event if added, null if none added.
+		UE_INTERNAL UE_API FSubsonicEvent* AddEvent(const FGameplayTag& InTag = { }, FSubsonicEvent NewEvent = {});
+#endif // WITH_EDITOR
+
+		// Assigns a universal CollectionId to the given collection.
+		UE_API UE_INTERNAL void AssignId();
+
+#if WITH_EDITOR
+		// Clears all actions under the given event
+		UE_INTERNAL UE_API bool ClearActions(const FEventHandle& InHandle);
+
+		// Removes all events in the given collection.
+		UE_INTERNAL UE_API void ClearEvents();
+
+		// Returns true if binding exists with the given parameter name on either the collection or event with
+		// the given name, false if not. Also returns false if event with the given name is not found.
+		UE_INTERNAL UE_API bool ContainsBinding(FName EventName, FName ParameterName) const;
+
+		// Fired when a structural parameter change creates stale editor bindings.
+		// The handler should show appropriate UI and return the desired response.
+		// If unbound, stale bindings are removed automatically.
+		UE_INTERNAL UE_API BindingUtils::FOnStaleBindingsDetected& GetOnStaleBindingsDetectedDelegate();
+#endif // WITH_EDITOR
+
+		UE_API bool ContainsEvent(FName Name) const;
+		UE_API bool ContainsEvent(const FGameplayTag& EventTag) const;
+		UE_API bool ContainsEvent(const FEventHandle& InHandle) const;
+
+		// Returns pointer to action at handle if it exists
+		UE_API const FSubsonicEventActionDefinition* FindAction(const FActionHandle& InHandle) const;
+
+		// Returns event with the given name.
+		UE_API const FSubsonicEvent* FindEvent(FName Name) const;
+		UE_API const FSubsonicEvent* FindEvent(const FGameplayTag& EventTag) const;
+		UE_API const FSubsonicEvent* FindEvent(const FEventHandle& InHandle) const;
+
+#if WITH_EDITOR
+		// Returns pointer to action at handle if it exists
+		UE_INTERNAL UE_API FSubsonicEventActionDefinition* FindMutableAction(const FActionHandle& InHandle);
+
+		// Returns mutable pointer to event at handle if it exists
+		UE_INTERNAL UE_API FSubsonicEvent* FindMutableEvent(const FEventHandle& InHandle);
+#endif // WITH_EDITOR
+
+		UE_API uint32 GetCollectionId() const;
+
+		UE_API const TMap<FGameplayTag, FSubsonicEvent>& GetEvents() const;
+
+		// Returns Id of invalid (default constructed or unregistered/reset collection).
+		UE_API static uint32 GetInvalidId();
+
+#if WITH_EDITOR
+		UE_API static FName GetEventsPropertyName();
+		UE_API static FName GetParametersPropertyName();
+#endif // WITH_EDITOR
+
+#if WITH_EDITORONLY_DATA
+		UE_API const FInstancedPropertyBag& GetParameters() const;
+#endif // WITH_EDITORONLY_DATA
+
+		UE_API bool IsValid() const;
+
+		// Register's the given collection with the SubsonicEventRegistry.
+		UE_API UE_INTERNAL void Register(const Core::FCollectionHandle& InHandle, Audio::FDeviceId InDeviceId);
+
+#if WITH_EDITOR
+		// Removes action at given handle and returns if successful.
+		UE_INTERNAL UE_API bool RemoveAction(const FActionHandle& InHandle);
+
+		// Removes event at given handle and returns if successful.
+		UE_INTERNAL UE_API bool RemoveEvent(const FEventHandle& InHandle);
+
+		// Removes all stale property bindings. Defined as stale if type or
+		// name of mapped property does not match that of target parameter.
+		// If bPromptRemoval is set, delegate is executed enabling editors to
+		// prompt user to cancel removal. Returns whether or not removal succeeded.
+		UE_INTERNAL UE_API bool RemoveStaleBindings(const FCollectionHandle& ParentHandle, bool bPromptRemoval = false);
+
+		// Binds all parameters across all events in the collection.
+		UE_INTERNAL UE_API void BindAllParameters();
+#endif // WITH_EDITOR
+
+		// Unregisters collection, leaving it in clear, invalid state.
+		UE_API UE_INTERNAL void Unregister();
+
+	private:
+#if WITH_EDITOR
+		BindingUtils::FOnStaleBindingsDetected OnStaleBindingsDetected;
+
+		UE_INTERNAL void BindParameter(FName Parameter);
+		UE_INTERNAL void UnbindParameter(FName Parameter);
+
+		// Removes all editor bindings to the given parameter name from all actions
+		// in all events. Returns whether or not any bindings were removed.
+		UE_INTERNAL UE_API bool RemoveAllPropertyBindingsToParameter(FName ParamName);
+#endif // WITH_EDITOR
+
+		// Unique Collection Id assigned post
+		// initialization of struct properties
+		// or via static initializer
+		// Independent from UObject identifier to
+		// allow for transient construction without
+		// having to utilize UObject wrapper.
+		uint32 CollectionId = INDEX_NONE;
+
+		// Device associated with the given event collection
+		Audio::FDeviceId DeviceId = INDEX_NONE;
+
+		UPROPERTY(EditAnywhere, BlueprintReadOnly, Category = "Events", meta = (AllowPrivateAccess))
+		TMap<FGameplayTag, FSubsonicEvent> Events;
+
+#if WITH_EDITORONLY_DATA
+		UPROPERTY(EditAnywhere, Category = "Parameters")
+		FInstancedPropertyBag Parameters;
+#endif // WITH_EDITORONLY_DATA
+	};
+
+} // namespace UE::Subsonic::Core
+
+// Required to be in global namespace still for UHT compatibility with namespaced Unreal reflected types
+template <>
+struct TStructOpsTypeTraits<UE::Subsonic::Core::FSubsonicEventCollectionDefinition>
+	: public TStructOpsTypeTraitsBase2<UE::Subsonic::Core::FSubsonicEventCollectionDefinition>
+{
+	enum { WithCopy = false };
+};
+#undef UE_API // SUBSONICCORE_API
diff --git a/Engine/Plugins/Experimental/Subsonic/Source/SubsonicCore/Public/SubsonicEventSubscriberInterface.h b/Engine/Plugins/Experimental/Subsonic/Source/SubsonicCore/Public/SubsonicEventSubscriberInterface.h
new file mode 100644
--- /dev/null
+++ b/Engine/Plugins/Experimental/Subsonic/Source/SubsonicCore/Public/SubsonicEventSubscriberInterface.h
@@ -0,0 +1,44 @@
+// Copyright Epic Games, Inc. All Rights Reserved.
+#pragma once
+
+#include "UObject/Interface.h"
+
+#include "SubsonicEventSubscriberInterface.generated.h"
+
+#define UE_API SUBSONICCORE_API
+
+namespace UE::Subsonic::Core
+{
+	// Forward Declarations
+	struct FCollectionHandle;
+	struct FEventHandle;
+	struct FSubsonicExecutor;
+
+	// Interface for all subscribers that respond to Subsonic events.
+	UINTERFACE(MinimalAPI, BlueprintType, meta = (DisplayName = "Subsonic Event Subscriber Interface", CannotImplementInterfaceInBlueprint))
+	class USubsonicEventSubscriberInterface : public UInterface
+	{
+		GENERATED_BODY()
+	};
+
+	class ISubsonicEventSubscriberInterface : public IInterface
+	{
+		GENERATED_BODY()
+
+	public:
+		UE_API virtual void OnCollectionRegistered(const FCollectionHandle& InCollection);
+		UE_API virtual void OnCollectionUnregistered(const FCollectionHandle& InCollection);
+
+		UE_API virtual void OnEventPreExecute(const FSubsonicExecutor& InExecutor, const FEventHandle& InHandle);
+		UE_API virtual void OnEventPostExecute(const FSubsonicExecutor& InExecutor, const FEventHandle& InHandle);
+
+		UE_API virtual void OnExecutorRegistered(const FSubsonicExecutor& InExecutor);
+		UE_API virtual void OnExecutorUnregistered(const FSubsonicExecutor& InExecutor);
+
+	protected:
+		UE_API virtual void Register();
+		UE_API virtual void Unregister();
+	};
+} // namespace UE::Subsonic::Core
+
+#undef UE_API // SUBSONICCORE_API
\ No newline at end of file
diff --git a/Engine/Plugins/Experimental/Subsonic/Source/SubsonicCore/Public/SubsonicExecutor.h b/Engine/Plugins/Experimental/Subsonic/Source/SubsonicCore/Public/SubsonicExecutor.h
new file mode 100644
--- /dev/null
+++ b/Engine/Plugins/Experimental/Subsonic/Source/SubsonicCore/Public/SubsonicExecutor.h
@@ -0,0 +1,131 @@
+// Copyright Epic Games, Inc. All Rights Reserved.
+#pragma once
+
+#include "AudioDefines.h"
+#include "HAL/Platform.h"
+#include "Misc/CoreDefines.h"
+#include "SubsonicHandles.h"
+#include "SubsonicParameterStore.h"
+#include "Templates/TypeHash.h"
+#include "Templates/UniquePtr.h"
+#include "UObject/NameTypes.h"
+
+#include "SubsonicExecutor.generated.h"
+
+
+// Forward Declarations
+struct FSubsonicEvent;
+struct FSubsonicEventCollectionDefinition;
+
+#define UE_API SUBSONICCORE_API
+
+
+UENUM(BlueprintType)
+enum class ESubsonicExecutionScope : uint8
+{
+	Global,
+	Collection,
+	Event,
+	Action
+};
+
+
+namespace UE::Subsonic::Core
+{
+	struct FSubsonicExecutor final
+	{
+	private:
+		FSubsonicExecutor();
+		FSubsonicExecutor(const FSubsonicExecutor&) = delete;
+
+	public:
+		FSubsonicExecutor(FSubsonicExecutor&&);
+		FSubsonicExecutor& operator=(const FSubsonicExecutor& InOther) = delete;
+		UE_API ~FSubsonicExecutor();
+
+		UE_API FSubsonicExecutor& operator=(FSubsonicExecutor&& InOther);
+
+		struct ICollectionAccessor
+		{
+			virtual ~ICollectionAccessor() = default;
+
+			virtual const FSubsonicEventCollectionDefinition* GetDefinition() const = 0;
+			virtual FCollectionHandle GetHandle() const = 0;
+		};
+		using FCollectionAccessorPtr = TUniquePtr<ICollectionAccessor>;
+
+		UE_API static FSubsonicExecutor Create(Audio::FDeviceId InDeviceId, FCollectionAccessorPtr&& CollectionPtr);
+
+		UE_API bool ExecuteEvent(FName EventName) const;
+
+		UE_API static uint32 GetInvalidId();
+		UE_API static FSubsonicExecutor GetInvalid();
+
+		UE_API FCollectionHandle GetCollectionHandle() const;
+		UE_API Audio::FDeviceId GetDeviceId() const;
+		UE_API uint32 GetId() const;
+
+		UE_API bool IsValid() const;
+
+		UE_API FString ToString() const;
+
+		UE_API void Unregister();
+
+		// Sets trigger-time parameters on this executor. Actions will merge these with their
+		// authored parameters when Execute() is called. Cleared on Unregister().
+		UE_API void SetParameters(FSubsonicParameterStore InParams);
+
+		// Returns the current trigger-time parameter store.
+		UE_API const FSubsonicParameterStore& GetParameters() const;
+
+	private:
+		void Register();
+
+		FCollectionAccessorPtr CollectionAccessor;
+
+		Audio::FDeviceId DeviceId;
+
+		uint32 ExecutorId = GetInvalidId();
+
+		// Trigger-time parameters stored here before ExecuteEvent is called.
+		FSubsonicParameterStore Parameters;
+	};
+
+	template <typename ScopeHandleType>
+	struct TScopeKey
+	{
+		uint32 ExecutorId = FSubsonicExecutor::GetInvalidId();
+		Audio::FDeviceId DeviceId = INDEX_NONE;
+		ScopeHandleType ScopeHandle;
+
+		TScopeKey() = default;
+		TScopeKey(const FSubsonicExecutor& InExecutor, ScopeHandleType InScopeHandle)
+			: ExecutorId(InExecutor.GetId())
+			, DeviceId(InExecutor.GetDeviceId())
+			, ScopeHandle(MoveTemp(InScopeHandle))
+		{
+		}
+
+		friend uint32 GetTypeHash(const TScopeKey<ScopeHandleType>& Key)
+		{
+			const uint32 Hash1 = HashCombineFast(Key.ExecutorId, GetTypeHash(Key.DeviceId));
+			return HashCombineFast(Hash1, GetTypeHash(Key.ScopeHandle));
+		}
+
+		friend bool operator==(const TScopeKey<ScopeHandleType>& A, const TScopeKey<ScopeHandleType>& B)
+		{
+			return A.ExecutorId == B.ExecutorId && A.ScopeHandle == B.ScopeHandle && A.DeviceId == B.DeviceId;
+		}
+
+		FString ToString() const
+		{
+			return FString::Printf(TEXT("%s [ExecId: %u, AudioDevice Id: %u]"), *ScopeHandle.ToString(), ExecutorId, DeviceId);
+		}
+	};
+
+	using FCollectionScopeKey = TScopeKey<FCollectionHandle>;
+	using FEventScopeKey = TScopeKey<FEventHandle>;
+	using FActionScopeKey = TScopeKey<FActionHandle>;
+} // namespace UE::Subsonic::Core
+
+#undef UE_API // SUBSONICCORE_API
diff --git a/Engine/Plugins/Experimental/Subsonic/Source/SubsonicCore/Public/SubsonicHandles.h b/Engine/Plugins/Experimental/Subsonic/Source/SubsonicCore/Public/SubsonicHandles.h
new file mode 100644
--- /dev/null
+++ b/Engine/Plugins/Experimental/Subsonic/Source/SubsonicCore/Public/SubsonicHandles.h
@@ -0,0 +1,87 @@
+// Copyright Epic Games, Inc. All Rights Reserved.
+#pragma once
+
+#include "HAL/Platform.h"
+#include "Misc/CoreDefines.h"
+#include "Templates/TypeHash.h"
+#include "UObject/NameTypes.h"
+
+#define UE_API SUBSONICCORE_API
+
+
+namespace UE::Subsonic::Core
+{
+	// Forward Declarations
+	struct FSubsonicEvent;
+	struct FSubsonicEventCollectionDefinition;
+
+	struct FCollectionHandle
+	{
+#if WITH_EDITORONLY_DATA
+		// Owning Collection definition's name (for logging/debugging)
+		FName CollectionName;
+#endif // WITH_EDITORONLY_DATA
+
+		// Returns invalid Id of a handle
+		UE_API static uint32 GetInvalidId();
+
+		// Owning Collection's instance id
+		uint32 CollectionId = GetInvalidId();
+
+		friend uint32 GetTypeHash(const FCollectionHandle& Handle)
+		{
+			return Handle.CollectionId;
+		}
+
+		friend bool operator==(const FCollectionHandle& A, const FCollectionHandle& B)
+		{
+			return A.CollectionId == B.CollectionId;
+		}
+
+		UE_API FString ToString() const;
+	};
+
+	struct FEventHandle
+	{
+		// Parent collection handle
+		FCollectionHandle Collection;
+
+		// Owning event name
+		FName EventName;
+
+		friend uint32 GetTypeHash(const FEventHandle& Handle)
+		{
+			return HashCombineFast(GetTypeHash(Handle.Collection), GetTypeHash(Handle.EventName));
+		}
+
+		friend bool operator==(const FEventHandle& A, const FEventHandle& B)
+		{
+			return A.Collection == B.Collection && A.EventName == B.EventName;
+		}
+
+		UE_API FString ToString() const;
+	};
+
+	struct FActionHandle
+	{
+		// Parent event handle
+		FEventHandle Event;
+
+		// Owning action's index within a given event
+		int32 Index = INDEX_NONE;
+
+		friend uint32 GetTypeHash(const FActionHandle& Handle)
+		{
+			return HashCombineFast(GetTypeHash(Handle.Event), GetTypeHash(Handle.Index));
+		}
+
+		friend bool operator==(const FActionHandle& A, const FActionHandle& B)
+		{
+			return A.Event == B.Event && A.Index == B.Index;
+		}
+
+		UE_API FString ToString() const;
+	};
+} // namespace UE::Subsonic::Core
+
+#undef UE_API // SUBSONICCORE_API
diff --git a/Engine/Plugins/Experimental/Subsonic/Source/SubsonicCore/Public/SubsonicParameterStore.h b/Engine/Plugins/Experimental/Subsonic/Source/SubsonicCore/Public/SubsonicParameterStore.h
new file mode 100644
--- /dev/null
+++ b/Engine/Plugins/Experimental/Subsonic/Source/SubsonicCore/Public/SubsonicParameterStore.h
@@ -0,0 +1,35 @@
+// Copyright Epic Games, Inc. All Rights Reserved.
+
+#pragma once
+
+#include "StructUtils/PropertyBag.h"
+
+#include "SubsonicParameterStore.generated.h"
+
+#define UE_API SUBSONICCORE_API
+
+
+namespace UE::Subsonic
+{
+	// A parameter store for Subsonic runtime parameters, backed by an FInstancedPropertyBag.
+	// Priority layers (lowest to highest): Authored < TriggerTime < Runtime.
+	USTRUCT(BlueprintType)
+	struct FSubsonicParameterStore
+	{
+		GENERATED_BODY()
+
+		UE_API bool HasParameter(FName Name) const;
+		UE_API void RemoveParameter(FName Name);
+		UE_API void Reset();
+		UE_API bool IsEmpty() const;
+
+		// Overlay: properties in Other are added to this store and their values overwrite any
+		// existing values with the same name.
+		UE_API void MergeFrom(const FSubsonicParameterStore& Other);
+
+		UPROPERTY(EditAnywhere, Category = "Parameters")
+		FInstancedPropertyBag Bag;
+	};
+} // namespace UE::Subsonic
+
+#undef UE_API
diff --git a/Engine/Plugins/Experimental/Subsonic/Source/SubsonicCore/SubsonicCore.Build.cs b/Engine/Plugins/Experimental/Subsonic/Source/SubsonicCore/SubsonicCore.Build.cs
new file mode 100644
--- /dev/null
+++ b/Engine/Plugins/Experimental/Subsonic/Source/SubsonicCore/SubsonicCore.Build.cs
@@ -0,0 +1,30 @@
+// Copyright Epic Games, Inc. All Rights Reserved.
+
+using UnrealBuildTool;
+
+public class SubsonicCore : ModuleRules
+{
+	public SubsonicCore(ReadOnlyTargetRules Target) : base(Target)
+	{
+		PCHUsage = ModuleRules.PCHUsageMode.UseExplicitOrSharedPCHs;
+
+		PublicDependencyModuleNames.AddRange(
+			new string[]
+			{
+				"AudioMixerCore",
+				"Core",
+				"GameplayTags",
+				"PropertyBindingUtils"
+			}
+		);
+
+		PrivateDependencyModuleNames.AddRange(
+			new string[]
+			{
+				"CoreUObject",
+			}
+		);
+
+		bAllowUETypesInNamespaces = true;
+	}
+}
diff --git a/Engine/Plugins/Experimental/Subsonic/Source/SubsonicEditor/Private/SubsonicAssetDefinitions.cpp b/Engine/Plugins/Experimental/Subsonic/Source/SubsonicEditor/Private/SubsonicAssetDefinitions.cpp
--- a/Engine/Plugins/Experimental/Subsonic/Source/SubsonicEditor/Private/SubsonicAssetDefinitions.cpp
+++ b/Engine/Plugins/Experimental/Subsonic/Source/SubsonicEditor/Private/SubsonicAssetDefinitions.cpp
@@ -3,8 +3,8 @@
 #include "SubsonicAssetDefinitions.h"
 
 #include "Styling/SlateStyleRegistry.h"
-#include "SubsonicEventCollection.h"
 #include "SubsonicEventCollectionEditor.h"
+#include "SubsonicEventCollectionObjects.h"
 
 #include UE_INLINE_GENERATED_CPP_BY_NAME(SubsonicAssetDefinitions)
 
diff --git a/Engine/Plugins/Experimental/Subsonic/Source/SubsonicEditor/Private/SubsonicAssetDefinitions.h b/Engine/Plugins/Experimental/Subsonic/Source/SubsonicEditor/Private/SubsonicAssetDefinitions.h
--- a/Engine/Plugins/Experimental/Subsonic/Source/SubsonicEditor/Private/SubsonicAssetDefinitions.h
+++ b/Engine/Plugins/Experimental/Subsonic/Source/SubsonicEditor/Private/SubsonicAssetDefinitions.h
@@ -3,6 +3,8 @@
 
 #include "AssetTypeActions/AssetDefinitionDefault_AudioDiffable.h"
 
+#include "SubsonicEventCollectionObjects.h"
+
 #include "SubsonicAssetDefinitions.generated.h"
 
 namespace UE::Subsonic
diff --git a/Engine/Plugins/Experimental/Subsonic/Source/SubsonicEditor/Private/SubsonicEditorSubsystem.cpp b/Engine/Plugins/Experimental/Subsonic/Source/SubsonicEditor/Private/SubsonicEditorSubsystem.cpp
--- a/Engine/Plugins/Experimental/Subsonic/Source/SubsonicEditor/Private/SubsonicEditorSubsystem.cpp
+++ b/Engine/Plugins/Experimental/Subsonic/Source/SubsonicEditor/Private/SubsonicEditorSubsystem.cpp
@@ -21,6 +21,8 @@ namespace UE::Subsonic
 
 	void USubsonicEditorSubsystem::RebuildActionStructChildCache()
 	{
+		using namespace UE::Subsonic::Core;
+
 		ActionStructs.Empty();
 		for (TObjectIterator<UScriptStruct> It; It; ++It)
 		{
diff --git a/Engine/Plugins/Experimental/Subsonic/Source/SubsonicEditor/Private/SubsonicEventCollectionEditor.cpp b/Engine/Plugins/Experimental/Subsonic/Source/SubsonicEditor/Private/SubsonicEventCollectionEditor.cpp
--- a/Engine/Plugins/Experimental/Subsonic/Source/SubsonicEditor/Private/SubsonicEventCollectionEditor.cpp
+++ b/Engine/Plugins/Experimental/Subsonic/Source/SubsonicEditor/Private/SubsonicEventCollectionEditor.cpp
@@ -2,11 +2,9 @@
 
 #include "SubsonicEventCollectionEditor.h"
 
-#include "Algo/AnyOf.h"
+#include "AudioDevice.h"
 #include "AudioDeviceHandle.h"
 #include "AudioDeviceManager.h"
-#include "CoreActions/SubsonicAction_GeneratorSource.h"
-#include "CoreEventSubscribers/SubsonicGeneratorSourceSubscriber.h"
 #include "Engine/Engine.h"
 #include "Framework/Application/SlateApplication.h"
 #include "Framework/Docking/TabManager.h"
@@ -14,16 +12,19 @@
 #include "GameplayTagContainer.h"
 #include "ISinglePropertyView.h"
 #include "Input/DragAndDrop.h"
-#include "Misc/MessageDialog.h"
 #include "Modules/ModuleManager.h"
 #include "PropertyEditorModule.h"
 #include "SGameplayTagCombo.h"
 #include "ScopedTransaction.h"
+#include "StandardActions/SubsonicAction_GeneratorSource.h"
+#include "StandardEventSubscribers/SubsonicGeneratorSourceSubscriber.h"
 #include "Styling/StyleColors.h"
 #include "Styling/SlateStyleRegistry.h"
 #include "SubsonicBindingUtils.h"
+#include "SubsonicDeviceUtils.h"
 #include "SubsonicEditorSubsystem.h"
 #include "SubsonicEventCollection.h"
+#include "SubsonicEventCollectionObjects.h"
 #include "SubsonicEventCollectionEditorCommands.h"
 #include "SubsonicExecutor.h"
 #include "SubsonicPropertyBindingExtension.h"
@@ -41,7 +42,7 @@
 
 namespace UE::Subsonic::Editor
 {
-	void TransactEventCollection(FText Description, USubsonicEventCollection& Collection, TFunctionRef<void(FSubsonicEventCollectionDefinition&)> TransactionFunc)
+	void TransactEventCollection(FText Description, USubsonicEventCollection& Collection, TFunctionRef<void(Core::FSubsonicEventCollectionDefinition&)> TransactionFunc)
 	{
 		FScopedTransaction Transaction(Description);
 		Collection.Modify();
@@ -61,9 +62,9 @@ namespace UE::Subsonic::Editor
 		public:
 			DRAG_DROP_OPERATOR_TYPE(FActionDragDropOp, FDragDropOperation)
 
-				FActionHandle SourceAction;
+			Core::FActionHandle SourceAction;
 
-			static TSharedRef<FActionDragDropOp> New(FActionHandle InAction)
+			static TSharedRef<FActionDragDropOp> New(Core::FActionHandle InAction)
 			{
 				TSharedRef<FActionDragDropOp> Op = MakeShared<FActionDragDropOp>();
 				Op->SourceAction = MoveTemp(InAction);
@@ -115,7 +116,7 @@ namespace UE::Subsonic::Editor
 				return { };
 			}
 
-			TSharedPtr<IPropertyHandle> EventsHandle = CollectionProperty->GetChildHandle(FSubsonicEventCollectionDefinition::GetEventsPropertyName());
+			TSharedPtr<IPropertyHandle> EventsHandle = CollectionProperty->GetChildHandle(Core::FSubsonicEventCollectionDefinition::GetEventsPropertyName());
 			if (!EventsHandle)
 			{
 				return { };
@@ -167,6 +168,8 @@ namespace UE::Subsonic::Editor
 		return SNew(SHandleVariantTableRow, TableView)
 		.OnDragDetected_Lambda([this, HandleVariant](const FGeometry& Geometry, const FPointerEvent& MouseEvent)
 		{
+			using namespace UE::Subsonic::Core;
+
 			if (const USubsonicEventCollection* Collection = GetEventCollection())
 			{
 				if (HandleVariant->IsAction())
@@ -182,6 +185,8 @@ namespace UE::Subsonic::Editor
 		})
 		.OnCanAcceptDrop_Lambda([](const FDragDropEvent& DragDropEvent, EItemDropZone InitZone, TSharedPtr<FHandleVariant> TargetVariant)
 		{
+			using namespace UE::Subsonic::Core;
+
 			TOptional<EItemDropZone> ToReturn;
 			if (TargetVariant.IsValid())
 			{
@@ -223,6 +228,8 @@ namespace UE::Subsonic::Editor
 
 	TSharedRef<ITableRow> FEventCollectionEditor::OnEventTree_GenerateRow(TSharedPtr<FEventCollectionEditor::FHandleVariant> HandleVariant, const TSharedRef<STableViewBase>& TableView)
 	{
+		using namespace UE::Subsonic::Core;
+
 		TSharedRef<SHandleVariantTableRow> NewRow = InitDraggableEventRow(HandleVariant, TableView);
 
 		const USubsonicEventCollection* Collection = GetEventCollection();
@@ -929,7 +936,8 @@ namespace UE::Subsonic::Editor
 
 	void FEventCollectionEditor::InitBindings()
 	{
-		using namespace BindingUtils;
+		using namespace UE::Subsonic::Core;
+		using namespace UE::Subsonic::Core::BindingUtils;
 
 		USubsonicEventCollection* Collection = GetEventCollection();
 		if (!Collection)
@@ -939,7 +947,7 @@ namespace UE::Subsonic::Editor
 
 		FSubsonicEventCollectionDefinition& Definition = Collection->GetMutableDefinition();
 		Definition.RemoveStaleBindings(Collection->GetHandle());
-		FOnStaleBindingsDetected& StaleBindingsDelegate = Definition.GetOnStaleBindingsDetectedDelegate();
+		BindingUtils::FOnStaleBindingsDetected& StaleBindingsDelegate = Definition.GetOnStaleBindingsDetectedDelegate();
 		StaleBindingsDelegate.Remove(StaleBindingsDelegateHandle);
 
 		// Register handler to show a confirmation dialog when a structural parameter change
@@ -951,7 +959,7 @@ namespace UE::Subsonic::Editor
 			{
 				if (const USubsonicEventCollection* PinnedCol = WeakCol.Get())
 				{
-					BindingUtils::PromptForStaleBindings(*PinnedCol, StaleBindings, OutResponse);
+					BindingUtils::PromptForStaleBindings(PinnedCol->GetDefinition(), StaleBindings, OutResponse);
 				}
 
 				// Deferred to next tick: the current transaction must finish posting before
@@ -1108,7 +1116,7 @@ namespace UE::Subsonic::Editor
 					.ButtonStyle(FAppStyle::Get(), "SimpleButton")
 					.OnClicked_Lambda([this]()
 					{
-						TransactEventCollection(LOCTEXT("AddEventDescription", "Add new Subsonic event to collection"), [&](FSubsonicEventCollectionDefinition& Definition)
+						TransactEventCollection(LOCTEXT("AddEventDescription", "Add new Subsonic event to collection"), [&](Core::FSubsonicEventCollectionDefinition& Definition)
 						{
 							if (FSubsonicEvent* NewEvent = Definition.AddEvent())
 							{
@@ -1136,7 +1144,7 @@ namespace UE::Subsonic::Editor
 					.ButtonStyle(FAppStyle::Get(), "SimpleButton")
 					.OnClicked_Lambda([this]()
 					{
-						TransactEventCollection(LOCTEXT("ClearAllEventsDescription", "Clear all Subsonic Collection Events"), [&](FSubsonicEventCollectionDefinition& Definition)
+						TransactEventCollection(LOCTEXT("ClearAllEventsDescription", "Clear all Subsonic Collection Events"), [&](Core::FSubsonicEventCollectionDefinition& Definition)
 						{
 							Definition.ClearEvents();
 							RebuildEventTree(true /* bRebuildListView */);
@@ -1224,8 +1232,10 @@ namespace UE::Subsonic::Editor
 		];
 	}
 
-	void FEventCollectionEditor::TransactEvent(const FEventHandle EventHandle, FText Description, TFunctionRef<void(FSubsonicEvent&)> TransactionFunc) const
+	void FEventCollectionEditor::TransactEvent(const FEventHandle EventHandle, FText Description, TFunctionRef<void(Core::FSubsonicEvent&)> TransactionFunc) const
 	{
+		using namespace UE::Subsonic::Core;
+
 		if (USubsonicEventCollection* Collection = GetEventCollection())
 		{
 			if (FSubsonicEvent* Event = Collection->GetMutableDefinition().FindMutableEvent(EventHandle))
@@ -1304,7 +1314,7 @@ namespace UE::Subsonic::Editor
 		return FAssetEditorToolkit::GetDefaultTabColor();
 	}
 
-	FSubsonicEventActionDefinition* FEventCollectionEditor::InsertAction(const FEventCollectionEditor::FHandleVariant& InHandleVariant)
+	Core::FSubsonicEventActionDefinition* FEventCollectionEditor::InsertAction(const FEventCollectionEditor::FHandleVariant& InHandleVariant)
 	{
 		FSubsonicEventActionDefinition* NewAction= nullptr;
 
@@ -1326,6 +1336,8 @@ namespace UE::Subsonic::Editor
 
 	void FEventCollectionEditor::MoveDroppedAction(const FActionHandle& SourceAction, const FHandleVariant& TargetHandle, EItemDropZone DropZone)
 	{
+		using namespace UE::Subsonic::Core;
+
 		if (TargetHandle.IsAction() && SourceAction == TargetHandle.GetAction())
 		{
 			return;
@@ -1500,6 +1512,8 @@ namespace UE::Subsonic::Editor
 
 	void FEventCollectionEditor::NotifyPostChange(const FPropertyChangedEvent& PropertyChangedEvent, FProperty* PropertyThatChanged)
 	{
+		using namespace UE::Subsonic::Core;
+
 		if (EventTreeInspectorView.IsValid())
 		{
 			EventTreeInspectorView->ForceRefresh();
@@ -1536,7 +1550,7 @@ namespace UE::Subsonic::Editor
 
 		const FSubsonicExecutor& Executor = PreviewExecutor->GetExecutor();
 		const FActionScopeKey ScopeKey(Executor, ActionHandle);
-		USubsonicGeneratorSourceSubscriber* Subscriber = ScopeKey.FindDeviceSubsystem<USubsonicGeneratorSourceSubscriber>();
+		USubsonicGeneratorSourceSubscriber* Subscriber = FindDeviceSubsystem<USubsonicGeneratorSourceSubscriber>(ScopeKey);
 		if (!Subscriber)
 		{
 			return;
diff --git a/Engine/Plugins/Experimental/Subsonic/Source/SubsonicEditor/Private/SubsonicEventCollectionEditor.h b/Engine/Plugins/Experimental/Subsonic/Source/SubsonicEditor/Private/SubsonicEventCollectionEditor.h
--- a/Engine/Plugins/Experimental/Subsonic/Source/SubsonicEditor/Private/SubsonicEventCollectionEditor.h
+++ b/Engine/Plugins/Experimental/Subsonic/Source/SubsonicEditor/Private/SubsonicEventCollectionEditor.h
@@ -1,17 +1,14 @@
 // Copyright Epic Games, Inc. All Rights Reserved.
 #pragma once
 
-#include "SubsonicEventCollectionEditor.h"
-
 #include "EditorUndoClient.h"
 #include "Framework/Docking/TabManager.h"
 #include "IDetailsView.h"
 #include "Misc/NotifyHook.h"
 #include "Misc/TVariant.h"
 #include "SubsonicEventCollection.h"
+#include "SubsonicEventCollectionObjects.h"
 #include "SubsonicEventCollectionViews.h"
-#include "SubsonicHandles.h"
-#include "Templates/Function.h"
 #include "Templates/SharedPointer.h"
 #include "Toolkits/AssetEditorToolkit.h"
 #include "Toolkits/IToolkitHost.h"
@@ -23,14 +20,19 @@
 
 namespace UE::Subsonic
 {
-	// Forward Declarations
-	class USubsonicEventCollection;
-
 	namespace Editor
 	{
 		class FEventCollectionEditor : public FAssetEditorToolkit, public FNotifyHook, public FEditorUndoClient
 		{
 		public:
+			using FActionHandle = Core::FActionHandle;
+			using FEventHandle = Core::FEventHandle;
+			using FCollectionHandle = Core::FCollectionHandle;
+
+			using FSubsonicEvent = Core::FSubsonicEvent;
+			using FSubsonicEventActionDefinition = Core::FSubsonicEventActionDefinition;
+			using FSubsonicEventCollectionDefinition = Core::FSubsonicEventCollectionDefinition;
+
 			static const FName EditorName;
 
 			FEventCollectionEditor() = default;
@@ -67,37 +69,37 @@ namespace UE::Subsonic
 			virtual void NotifyPostChange(const FPropertyChangedEvent& PropertyChangedEvent, FProperty* PropertyThatChanged) override;
 
 		private:
-			void TransactEvent(const FEventHandle EventHandle, FText Description, TFunctionRef<void(FSubsonicEvent&)> TransactionFunc) const;
+			void TransactEvent(const Core::FEventHandle EventHandle, FText Description, TFunctionRef<void(FSubsonicEvent&)> TransactionFunc) const;
 			void TransactEventCollection(FText Description, TFunctionRef<void(FSubsonicEventCollectionDefinition&)> TransactionFunc) const;
 
 			TSharedRef<SWidget> InitEventTreeView();
 
 			class FHandleVariant
 			{
-				TVariant<FEventHandle, FActionHandle> Handle;
+				TVariant<Core::FEventHandle, Core::FActionHandle> Handle;
 
 			public:
-				const FActionHandle& GetAction() const
+				const Core::FActionHandle& GetAction() const
 				{
 					checkf(IsAction(), TEXT("Must check that handle is action first"));
-					return Handle.Get<FActionHandle>();
+					return Handle.Get<Core::FActionHandle>();
 				}
 
-				const FEventHandle& GetEvent() const
+				const Core::FEventHandle& GetEvent() const
 				{
-					return Handle.IsType<FActionHandle>()
-						? Handle.Get<FActionHandle>().Event
-						: Handle.Get<FEventHandle>();
+					return Handle.IsType<Core::FActionHandle>()
+						? Handle.Get<Core::FActionHandle>().Event
+						: Handle.Get<Core::FEventHandle>();
 				}
 
 				bool IsAction() const
 				{
-					return Handle.IsType<FActionHandle>();
+					return Handle.IsType<Core::FActionHandle>();
 				}
 
 				bool IsEvent() const
 				{
-					return Handle.IsType<FEventHandle>();
+					return Handle.IsType<Core::FEventHandle>();
 				}
 
 				template <typename HandleType>
@@ -109,14 +111,14 @@ namespace UE::Subsonic
 
 			// Inserts action. If the handle variant points at an event, adds to end of event array.
 			// If it points to an action, inserts before the given action.
-			FSubsonicEventActionDefinition* InsertAction(const FHandleVariant& InHandleVariant);
+			Core::FSubsonicEventActionDefinition* InsertAction(const FHandleVariant& InHandleVariant);
 
 			// Moves an action to the position relative to TargetAction per the given DropZone.
-			void MoveDroppedAction(const FActionHandle& SourceAction, const FHandleVariant& HandleVariant, EItemDropZone DropZone);
+			void MoveDroppedAction(const Core::FActionHandle& SourceAction, const FHandleVariant& HandleVariant, EItemDropZone DropZone);
 
 			void BuildTransport();
 
-			void ExecutePreviewEvent(const FEventHandle& EventHandle);
+			void ExecutePreviewEvent(const Core::FEventHandle& EventHandle);
 
 			void InitBindings();
 
diff --git a/Engine/Plugins/Experimental/Subsonic/Source/SubsonicEditor/Private/SubsonicEventCollectionFactory.cpp b/Engine/Plugins/Experimental/Subsonic/Source/SubsonicEditor/Private/SubsonicEventCollectionFactory.cpp
--- a/Engine/Plugins/Experimental/Subsonic/Source/SubsonicEditor/Private/SubsonicEventCollectionFactory.cpp
+++ b/Engine/Plugins/Experimental/Subsonic/Source/SubsonicEditor/Private/SubsonicEventCollectionFactory.cpp
@@ -1,7 +1,7 @@
 // Copyright Epic Games, Inc. All Rights Reserved.
 #include "SubsonicEventCollectionFactory.h"
 
-#include "SubsonicEventCollection.h"
+#include "SubsonicEventCollectionObjects.h"
 
 #include UE_INLINE_GENERATED_CPP_BY_NAME(SubsonicEventCollectionFactory)
 
diff --git a/Engine/Plugins/Experimental/Subsonic/Source/SubsonicEditor/Private/SubsonicEventCollectionFactory.h b/Engine/Plugins/Experimental/Subsonic/Source/SubsonicEditor/Private/SubsonicEventCollectionFactory.h
--- a/Engine/Plugins/Experimental/Subsonic/Source/SubsonicEditor/Private/SubsonicEventCollectionFactory.h
+++ b/Engine/Plugins/Experimental/Subsonic/Source/SubsonicEditor/Private/SubsonicEventCollectionFactory.h
@@ -5,6 +5,8 @@
 #include "Factories/Factory.h"
 #include "UObject/Object.h"
 
+#include "SubsonicEventCollectionObjects.h"
+
 #include "SubsonicEventCollectionFactory.generated.h"
 
 
diff --git a/Engine/Plugins/Experimental/Subsonic/Source/SubsonicEditor/Private/SubsonicEventCollectionViews.cpp b/Engine/Plugins/Experimental/Subsonic/Source/SubsonicEditor/Private/SubsonicEventCollectionViews.cpp
--- a/Engine/Plugins/Experimental/Subsonic/Source/SubsonicEditor/Private/SubsonicEventCollectionViews.cpp
+++ b/Engine/Plugins/Experimental/Subsonic/Source/SubsonicEditor/Private/SubsonicEventCollectionViews.cpp
@@ -5,6 +5,7 @@
 #include "DetailCategoryBuilder.h"
 #include "Internationalization/Internationalization.h"
 #include "SubsonicEventCollection.h"
+#include "SubsonicEventCollectionObjects.h"
 
 #include UE_INLINE_GENERATED_CPP_BY_NAME(SubsonicEventCollectionViews)
 
@@ -17,6 +18,8 @@ namespace UE::Subsonic::Editor
 	{
 		void AddActionProperties(IDetailLayoutBuilder& LayoutBuilder, IDetailCategoryBuilder& CategoryBuilder, TSharedRef<IPropertyHandleArray> ActionArrayHandle, uint32 ActionIndex)
 		{
+			using namespace UE::Subsonic::Core;
+
 			TSharedRef<IPropertyHandle> ActionEntryHandle = ActionArrayHandle->GetElement(ActionIndex);
 			uint32 NumProperties = INDEX_NONE;
 			if (ActionEntryHandle->GetNumChildren(NumProperties) == FPropertyAccess::Success)
@@ -35,6 +38,8 @@ namespace UE::Subsonic::Editor
 
 		TSharedPtr<IPropertyHandle> GetEventHandleMatchingName(TSharedRef<IPropertyHandleMap> EventMapHandle, uint32 Index, FName Name)
 		{
+			using namespace UE::Subsonic::Core;
+
 			TSharedRef<IPropertyHandle> EventEntryHandle = EventMapHandle->GetElement(Index);
 			if (TSharedPtr<IPropertyHandle> EventKeyHandle = EventEntryHandle->GetKeyHandle(); EventKeyHandle.IsValid())
 			{
@@ -54,6 +59,8 @@ namespace UE::Subsonic::Editor
 
 	void FCollectionParametersDetailCustomization::CustomizeDetails(IDetailLayoutBuilder& LayoutBuilder)
 	{
+		using namespace UE::Subsonic::Core;
+
 		if (USubsonicEventCollectionViewBase* View = GetViewFromBuilder(LayoutBuilder))
 		{
 			const FName ParametersPropertyName = FSubsonicEventCollectionDefinition::GetParametersPropertyName();
@@ -115,6 +122,8 @@ namespace UE::Subsonic::Editor
 
 	void FEventTreeSelectionDetailCustomization::CustomizeDetails(IDetailLayoutBuilder& LayoutBuilder)
 	{
+		using namespace UE::Subsonic::Core;
+
 		USubsonicEventTreeDetailsView* View = Cast<USubsonicEventTreeDetailsView>(GetViewFromBuilder(LayoutBuilder));
 		if (!View)
 		{
diff --git a/Engine/Plugins/Experimental/Subsonic/Source/SubsonicEditor/Private/SubsonicEventCollectionViews.h b/Engine/Plugins/Experimental/Subsonic/Source/SubsonicEditor/Private/SubsonicEventCollectionViews.h
--- a/Engine/Plugins/Experimental/Subsonic/Source/SubsonicEditor/Private/SubsonicEventCollectionViews.h
+++ b/Engine/Plugins/Experimental/Subsonic/Source/SubsonicEditor/Private/SubsonicEventCollectionViews.h
@@ -4,6 +4,7 @@
 #include "DetailLayoutBuilder.h"
 #include "IDetailCustomization.h"
 #include "SubsonicEventCollection.h"
+#include "SubsonicEventCollectionObjects.h"
 
 #include "SubsonicEventCollectionViews.generated.h"
 
diff --git a/Engine/Plugins/Experimental/Subsonic/Source/SubsonicEditor/Private/SubsonicPropertyBindingExtension.cpp b/Engine/Plugins/Experimental/Subsonic/Source/SubsonicEditor/Private/SubsonicPropertyBindingExtension.cpp
--- a/Engine/Plugins/Experimental/Subsonic/Source/SubsonicEditor/Private/SubsonicPropertyBindingExtension.cpp
+++ b/Engine/Plugins/Experimental/Subsonic/Source/SubsonicEditor/Private/SubsonicPropertyBindingExtension.cpp
@@ -42,6 +42,8 @@ namespace UE::Subsonic::Editor
 		const FProperty* TargetProperty,
 		FDetailWidgetRow& WidgetRow)
 	{
+		using namespace UE::Subsonic::Core;
+
 		check(TargetProperty);
 
 		// Weak refs used inside all delegates to guard against stale pointers.
@@ -139,6 +141,8 @@ namespace UE::Subsonic::Editor
 		TWeakObjectPtr<USubsonicEventCollection> WeakCollection,
 		TWeakObjectPtr<USubsonicEventTreeDetailsView> WeakView)
 	{
+		using namespace UE::Subsonic::Core;
+
 		constexpr bool bInShouldCloseWindowAfterMenuSelection = true;
 		TSharedPtr<const FUICommandList> NullCommandList;
 		FMenuBuilder MenuBuilder(bInShouldCloseWindowAfterMenuSelection, NullCommandList);
@@ -295,6 +299,8 @@ namespace UE::Subsonic::Editor
 		TWeakObjectPtr<USubsonicEventCollection> WeakCollection,
 		TWeakObjectPtr<USubsonicEventTreeDetailsView> WeakView)
 	{
+		using namespace UE::Subsonic::Core;
+
 		const FName BoundParam = GetCurrentBindingParam();
 		if (BoundParam.IsNone())
 		{
@@ -336,6 +342,8 @@ namespace UE::Subsonic::Editor
 		TWeakObjectPtr<USubsonicEventTreeDetailsView> WeakView,
 		const FProperty* TargetProperty)
 	{
+		using namespace UE::Subsonic::Core;
+
 		USubsonicEventCollection* Collection = WeakCollection.Get();
 		USubsonicEventTreeDetailsView* DetailsView = WeakView.Get();
 		if (Collection && DetailsView)
@@ -359,6 +367,8 @@ namespace UE::Subsonic::Editor
 
 	void FSubsonicPropertyBindingExtension::ExtendWidgetRow(FDetailWidgetRow& WidgetRow, const IDetailLayoutBuilder& InDetailBuilder, const UClass* InObjectClass, TSharedPtr<IPropertyHandle> InPropertyHandle)
 	{
+		using namespace UE::Subsonic::Core;
+
 		if (!IModularFeatures::Get().IsModularFeatureAvailable("PropertyAccessEditor"))
 		{
 			return;
@@ -434,6 +444,8 @@ namespace UE::Subsonic::Editor
 
 	void FSubsonicPropertyBindingExtension::AddParamBinding(USubsonicEventCollection* Collection, USubsonicEventTreeDetailsView* DetailsView, FName PropertyName, const int32 SourceStructIndex, const FProperty* SourceProperty)
 	{
+		using namespace UE::Subsonic::Core;
+
 		if (!SourceProperty || !Collection || !DetailsView || DetailsView->ActionIndex < 0)
 		{
 			return;
diff --git a/Engine/Plugins/Experimental/Subsonic/Source/SubsonicEditor/Private/SubsonicPropertyBindingExtension.h b/Engine/Plugins/Experimental/Subsonic/Source/SubsonicEditor/Private/SubsonicPropertyBindingExtension.h
--- a/Engine/Plugins/Experimental/Subsonic/Source/SubsonicEditor/Private/SubsonicPropertyBindingExtension.h
+++ b/Engine/Plugins/Experimental/Subsonic/Source/SubsonicEditor/Private/SubsonicPropertyBindingExtension.h
@@ -11,7 +11,11 @@ namespace UE::Subsonic
 {
 	// Forward Declarations
 	class USubsonicEventCollection;
-	struct FSubsonicEventCollectionDefinition;
+
+	namespace Core
+	{
+		struct FSubsonicEventCollectionDefinition;
+	} // namespace Core
 
 	namespace Editor
 	{
@@ -37,7 +41,7 @@ namespace UE::Subsonic
 			}
 		};
 
-		void TransactEventCollection(FText Description, USubsonicEventCollection& Collection, TFunctionRef<void(FSubsonicEventCollectionDefinition&)> TransactionFunc);
+		void TransactEventCollection(FText Description, USubsonicEventCollection& Collection, TFunctionRef<void(Core::FSubsonicEventCollectionDefinition&)> TransactionFunc);
 
 		// Only shows the binding dropdown on properties not decorated with 'NoBinding' Metadata
 		// tag. Overrides ExtendWidgetRow to populate the binding menu with parameters from the
diff --git a/Engine/Plugins/Experimental/Subsonic/Source/SubsonicEditor/SubsonicEditor.Build.cs b/Engine/Plugins/Experimental/Subsonic/Source/SubsonicEditor/SubsonicEditor.Build.cs
--- a/Engine/Plugins/Experimental/Subsonic/Source/SubsonicEditor/SubsonicEditor.Build.cs
+++ b/Engine/Plugins/Experimental/Subsonic/Source/SubsonicEditor/SubsonicEditor.Build.cs
@@ -15,7 +15,7 @@ public class SubsonicEditor : ModuleRules
 				"CoreUObject"
 			}
 		);
-			
+
 		PrivateDependencyModuleNames.AddRange(
 			new string[]
 			{
@@ -35,14 +35,13 @@ public class SubsonicEditor : ModuleRules
 				"PropertyEditor",
 				"Slate",
 				"SlateCore",
-				"Subsonic",
+				"SubsonicCore",
+				"SubsonicEngine",
 				"ToolWidgets",
 				"UnrealEd"
 			}
 		);
 
-		PrivateIncludePaths.Add(System.IO.Path.Combine(ModuleDirectory, "..", "Subsonic", "Private"));
-
 		bAllowUETypesInNamespaces = true;
 	}
 }
diff --git a/Engine/Plugins/Experimental/Subsonic/Source/SubsonicEngine/Private/StandardActions/SubsonicAction_AudioComponent.cpp b/Engine/Plugins/Experimental/Subsonic/Source/SubsonicEngine/Private/StandardActions/SubsonicAction_AudioComponent.cpp
new file mode 100644
--- /dev/null
+++ b/Engine/Plugins/Experimental/Subsonic/Source/SubsonicEngine/Private/StandardActions/SubsonicAction_AudioComponent.cpp
@@ -0,0 +1,248 @@
+// Copyright Epic Games, Inc. All Rights Reserved.
+
+#include "StandardActions/SubsonicAction_AudioComponent.h"
+
+#include "IAudioModulation.h"
+#include "Internationalization/Internationalization.h"
+#include "Sound/SoundBase.h"
+#include "Sound/SoundConcurrency.h"
+#include "SubsonicDeviceUtils.h"
+#include "SubsonicExecutor.h"
+#include "SubsonicHandles.h"
+
+
+#define LOCTEXT_NAMESPACE "SubsonicActions"
+
+namespace UE::Subsonic
+{
+	namespace AudioComponentPrivate
+	{
+		UAudioComponent* AddScopedComponent(const Core::FActionScopeKey& InKey, ESubsonicExecutionScope Scope, FName Name)
+		{
+			UAudioComponent* Component = nullptr;
+			if (USubsonicAudioComponentSubscriber* AudioComponentSubscriber = FindDeviceSubsystem<USubsonicAudioComponentSubscriber>(InKey))
+			{
+				switch (Scope)
+				{
+					case ESubsonicExecutionScope::Global:
+					{
+						Component = AudioComponentSubscriber->AddComponent(Name);
+					}
+					break;
+
+					case ESubsonicExecutionScope::Collection:
+					{
+						Component = AudioComponentSubscriber->AddComponent(InKey.ScopeHandle.Event.Collection, Name);
+					}
+					break;
+
+					case ESubsonicExecutionScope::Event:
+					{
+						Component = AudioComponentSubscriber->AddComponent(InKey.ScopeHandle.Event, Name);
+					}
+					break;
+
+					case ESubsonicExecutionScope::Action:
+					{
+						Component = AudioComponentSubscriber->AddComponent(InKey.ScopeHandle, Name);
+					}
+					break;
+				}
+			}
+
+			return Component;
+		}
+
+		UAudioComponent* FindScopedComponent(const Core::FActionScopeKey& InKey, ESubsonicExecutionScope Scope, FName Name)
+		{
+			UAudioComponent* Component = nullptr;
+			if (USubsonicAudioComponentSubscriber* AudioComponentSubscriber = FindDeviceSubsystem<USubsonicAudioComponentSubscriber>(InKey))
+			{
+				switch (Scope)
+					{
+					case ESubsonicExecutionScope::Global:
+					{
+						Component = AudioComponentSubscriber->FindComponent(Name);
+					}
+					break;
+
+					case ESubsonicExecutionScope::Collection:
+					{
+						Component = AudioComponentSubscriber->FindComponent(InKey.ScopeHandle.Event.Collection, Name);
+					}
+					break;
+
+					case ESubsonicExecutionScope::Event:
+					{
+						Component = AudioComponentSubscriber->FindComponent(InKey.ScopeHandle.Event, Name);
+					}
+					break;
+
+					case ESubsonicExecutionScope::Action:
+					{
+						Component = AudioComponentSubscriber->FindComponent(InKey.ScopeHandle, Name);
+					}
+					break;
+				}
+			}
+
+			return Component;
+		}
+
+		UAudioComponent* AccessComponent(const Core::FActionScopeKey& InKey, ESubsonicAudioComponentAccess Access, ESubsonicExecutionScope Scope, FName Name)
+		{
+			UAudioComponent* Component = nullptr;
+			if (Access == ESubsonicAudioComponentAccess::Find || Access == ESubsonicAudioComponentAccess::FindOrAdd)
+			{
+				Component = FindScopedComponent(InKey, Scope, Name);
+			}
+
+			bool bAddNewComponent = Access == ESubsonicAudioComponentAccess::Add;
+			if (Access == ESubsonicAudioComponentAccess::FindOrAdd)
+			{
+				// Didn't find an existing component above, so attempt
+				// to make one if set to find or add even if access not set to 'Add'.
+				if (!Component)
+				{
+					bAddNewComponent = true;
+				}
+			}
+
+			if (bAddNewComponent)
+			{
+				Component = AddScopedComponent(InKey, Scope, Name);
+			}
+
+			return Component;
+		}
+	} // namespace AudioComponentPrivate
+
+	void FSubsonicEventAction_AudioComponentModify::Execute(const Core::FSubsonicExecutor& InExecutor, const Core::FActionHandle& InHandle) const
+	{
+		const Core::FActionScopeKey ScopeKey(InExecutor, InHandle);
+		if (UAudioComponent* AudioComponent = AudioComponentPrivate::AccessComponent(ScopeKey, Access, Scope, Name))
+		{
+			for (const TInstancedStruct<FSubsonicEventAction_AudioComponentModifierBase>& Modifier : Modifiers)
+			{
+				if (const FSubsonicEventAction_AudioComponentModifierBase* ModPtr = Modifier.GetPtr())
+				{
+					ModPtr->Execute(*AudioComponent);
+				}
+			}
+		}
+	}
+
+#if WITH_EDITOR
+	FText FSubsonicEventAction_AudioComponentModify::GetDisplayInfo() const
+	{
+		FText NameText = FText::FromName(Name);
+		if (NameText.IsEmptyOrWhitespace())
+		{
+			NameText = Super::GetDisplayInfo();
+		}
+
+		if (!Modifiers.IsEmpty())
+		{
+			if (const UScriptStruct* ModifierStruct = Modifiers[0].GetScriptStruct())
+			{
+				if (ModifierStruct&& Modifiers.Num() > 1)
+				{
+					return FText::Format(LOCTEXT("AudioComponentModifierListDisplaySinglePlusInfo_Format", "{0} ('{1}' and {2} more)"),
+						NameText,
+						ModifierStruct->GetDisplayNameText(),
+						FText::AsNumber(Modifiers.Num() - 1));
+				}
+
+				return FText::Format(LOCTEXT("AudioComponentModifierListDisplayInfoSingle_Format", "{0} ({1})"),
+					NameText,
+					ModifierStruct->GetDisplayNameText());
+			}
+
+			return FText::Format(LOCTEXT("AudioComponentModifierListDisplayInfoNum_Format", "{0} ({1} modifiers)"),
+				NameText,
+				FText::AsNumber(Modifiers.Num()));
+		}
+
+		return NameText;
+	}
+	
+	FText FSubsonicEventAction_AudioComponentPlay::GetDisplayInfo() const
+	{
+		if (!Name.IsNone())
+		{
+			if (Sound != nullptr)
+			{
+				return FText::Format(LOCTEXT("AudioComponentPlay_DisplayInfoFormat", "{0} ({1})"), FText::FromName(Name), FText::FromName(Sound->GetFName()));
+			}
+
+			return FText::FromName(Name);
+		}
+
+		return Super::GetDisplayInfo();
+	}
+#endif // WITH_EDITOR
+
+	void FSubsonicEventAction_AudioComponentPlay::Execute(const Core::FSubsonicExecutor& InExecutor, const Core::FActionHandle& InHandle) const
+	{
+		const Core::FActionScopeKey ScopeKey(InExecutor, InHandle);
+		if (UAudioComponent* AudioComponent = AudioComponentPrivate::AccessComponent(ScopeKey, Access, Scope, Name))
+		{
+			AudioComponent->SetSound(Sound);
+			AudioComponent->Play();
+		}
+	}
+
+	void FSubsonicEventAction_AudioComponentStop::Execute(const Core::FSubsonicExecutor& InExecutor, const Core::FActionHandle& InHandle) const
+	{
+		const Core::FActionScopeKey ScopeKey(InExecutor, InHandle);
+		if (UAudioComponent* AudioComponent = AudioComponentPrivate::FindScopedComponent(ScopeKey, Scope, Name))
+		{
+			AudioComponent->Stop();
+		}
+	}
+
+#if WITH_EDITOR
+	FText FSubsonicEventAction_AudioComponentStop::GetDisplayInfo() const
+	{
+		if (!Name.IsNone())
+		{
+			return FText::FromName(Name);
+		}
+
+		return Super::GetDisplayInfo();
+	}
+#endif // WITH_EDITOR
+
+	void FSubsonicEventAction_AudioComponentModifier_Play::Execute(UAudioComponent& AudioComponent) const
+	{
+		AudioComponent.Play(StartTime);
+	}
+
+	void FSubsonicEventAction_AudioComponentModifier_SetAttenuation::Execute(UAudioComponent& AudioComponent) const
+	{
+		AudioComponent.SetAttenuationSettings(Attenuation);
+	}
+
+	void FSubsonicEventAction_AudioComponentModifier_SetConcurrency::Execute(UAudioComponent& AudioComponent) const
+	{
+		AudioComponent.ConcurrencySet = Concurrency;
+	}
+
+	void FSubsonicEventAction_AudioComponentModifier_SetModulationRouting::Execute(UAudioComponent& AudioComponent) const
+	{
+		TSet<USoundModulatorBase*> ModulatorSet;
+		Algo::Transform(Modulators, ModulatorSet, [](const TObjectPtr<USoundModulatorBase>& Modulator) { return Modulator.Get(); });
+		AudioComponent.SetModulationRouting(ModulatorSet, Destination, RoutingMethod);
+	}
+
+	void FSubsonicEventAction_AudioComponentModifier_SetSound::Execute(UAudioComponent& AudioComponent) const
+	{
+		AudioComponent.SetSound(Sound);
+	}
+
+	void FSubsonicEventAction_AudioComponentModifier_Stop::Execute(UAudioComponent& AudioComponent) const
+	{
+		AudioComponent.Stop();
+	}
+} // namespace UE::Subsonic
+#undef LOCTEXT_NAMESPACE // SubsonicActions
diff --git a/Engine/Plugins/Experimental/Subsonic/Source/SubsonicEngine/Private/StandardActions/SubsonicAction_AudioComponent.h b/Engine/Plugins/Experimental/Subsonic/Source/SubsonicEngine/Private/StandardActions/SubsonicAction_AudioComponent.h
new file mode 100644
--- /dev/null
+++ b/Engine/Plugins/Experimental/Subsonic/Source/SubsonicEngine/Private/StandardActions/SubsonicAction_AudioComponent.h
@@ -0,0 +1,223 @@
+// Copyright Epic Games, Inc. All Rights Reserved.
+#pragma once
+
+#include "Sound/SoundModulationDestination.h"
+#include "StandardEventSubscribers/AudioComponentEventSubscriber.h"
+#include "StructUtils/InstancedStruct.h"
+#include "SubsonicEventCollection.h"
+#include "SubsonicExecutor.h"
+
+#include "SubsonicAction_AudioComponent.generated.h"
+
+
+// Forward Declarations
+class UAudioComponent;
+class USoundConcurrency;
+class USoundModulatorBase;
+
+struct FSubsonicEventAction_AudioComponentModifierBase;
+
+
+namespace UE::Subsonic
+{
+	UENUM()
+	enum class ESubsonicAudioComponentAccess : uint8
+	{
+		// Adds new component, releasing any existing component.
+		// Stops existing component if sound is not one-shot.
+		Add,
+
+		// Finds existing component, or creates a new one.
+		FindOrAdd,
+
+		// Ignore request and if component doesn't already exist.
+		Find
+	};
+
+	USTRUCT(MinimalAPI, DisplayName = "Modify Audio Component")
+	struct FSubsonicEventAction_AudioComponentModify : public Core::FSubsonicEventActionBase
+	{
+		GENERATED_BODY()
+
+	public:
+		virtual ~FSubsonicEventAction_AudioComponentModify() = default;
+
+		virtual void Execute(const Core::FSubsonicExecutor& InExecutor, const Core::FActionHandle& InHandle) const override;
+
+#if WITH_EDITOR
+		virtual FText GetDisplayInfo() const override;
+#endif // WITH_EDITOR
+
+		// Addressable name of audio component instance.
+		UPROPERTY(EditAnywhere, Category = "Subsonic")
+		FName Name = "Component01";
+
+		// Scope of component accessor.  If set to event, provided name is used
+		// to address the same component within the event, collection the same component
+		// within the collection, etc.
+		UPROPERTY(EditAnywhere, Category = "Subsonic")
+		ESubsonicExecutionScope Scope = ESubsonicExecutionScope::Collection;
+
+		UPROPERTY(EditAnywhere, Category = "Subsonic")
+		ESubsonicAudioComponentAccess Access = ESubsonicAudioComponentAccess::FindOrAdd;
+
+		UPROPERTY(EditAnywhere, Category = "Subsonic")
+		TArray<TInstancedStruct<FSubsonicEventAction_AudioComponentModifierBase>> Modifiers;
+	};
+
+	USTRUCT(MinimalAPI, DisplayName = "Play Sound")
+	struct FSubsonicEventAction_AudioComponentPlay : public Core::FSubsonicEventActionBase
+	{
+		GENERATED_BODY()
+
+		virtual ~FSubsonicEventAction_AudioComponentPlay() = default;
+
+#if WITH_EDITOR
+		virtual FText GetDisplayInfo() const override;
+#endif // WITH_EDITOR
+
+		virtual void Execute(const Core::FSubsonicExecutor& InExecutor, const Core::FActionHandle& InHandle) const override;
+
+	public:
+		// Addressable name of audio component instance.
+		UPROPERTY(EditAnywhere, Category = "Subsonic")
+		FName Name = "Component01";
+
+		// Sound to play.
+		UPROPERTY(EditAnywhere, Category = "Subsonic")
+		TObjectPtr<USoundBase> Sound;
+
+		// Scope of component accessor.  If set to event, provided name is used
+		// to address the same component within the event, collection the same component
+		// within the collection, etc.
+		UPROPERTY(EditAnywhere, Category = "Subsonic")
+		ESubsonicExecutionScope Scope = ESubsonicExecutionScope::Collection;
+
+		UPROPERTY(EditAnywhere, Category = "Subsonic")
+		ESubsonicAudioComponentAccess Access = ESubsonicAudioComponentAccess::FindOrAdd;
+	};
+
+	USTRUCT(MinimalAPI, DisplayName = "Stop Sound")
+	struct FSubsonicEventAction_AudioComponentStop : public Core::FSubsonicEventActionBase
+	{
+		GENERATED_BODY()
+
+		virtual ~FSubsonicEventAction_AudioComponentStop() = default;
+
+		virtual void Execute(const Core::FSubsonicExecutor& InExecutor, const Core::FActionHandle& InHandle) const override;
+
+#if WITH_EDITOR
+		virtual FText GetDisplayInfo() const override;
+#endif // WITH_EDITOR
+
+		// Addressable name of component/sound instance.
+		UPROPERTY(EditAnywhere, Category = "Subsonic")
+		FName Name = "Component01";
+
+		// Scope of component to play.  If set to event, provided name is used
+		// to address the existing component within the event, collection the existing component
+		// within the collection, etc.
+		UPROPERTY(EditAnywhere, Category = "Subsonic")
+		ESubsonicExecutionScope Scope = ESubsonicExecutionScope::Collection;
+	};
+
+	USTRUCT(MinimalAPI, meta = (Abstract, Hidden))
+	struct FSubsonicEventAction_AudioComponentModifierBase
+	{
+		GENERATED_BODY()
+
+	public:
+		virtual ~FSubsonicEventAction_AudioComponentModifierBase() = default;
+
+		virtual void Execute(UAudioComponent& AudioComponent) const { }
+	};
+
+	USTRUCT(MinimalAPI, DisplayName = "Play")
+	struct FSubsonicEventAction_AudioComponentModifier_Play : public FSubsonicEventAction_AudioComponentModifierBase
+	{
+		GENERATED_BODY()
+
+	public:
+		virtual ~FSubsonicEventAction_AudioComponentModifier_Play() = default;
+
+		virtual void Execute(UAudioComponent& AudioComponent) const override;
+
+		// Play StartTime Offset (plays immediately if 0.0)
+		UPROPERTY(EditAnywhere, Category = "Subsonic")
+		float StartTime = 0.0f;
+	};
+
+	USTRUCT(MinimalAPI, DisplayName = "Set Attenuation")
+	struct FSubsonicEventAction_AudioComponentModifier_SetAttenuation : public FSubsonicEventAction_AudioComponentModifierBase
+	{
+		GENERATED_BODY()
+
+	public:
+		virtual ~FSubsonicEventAction_AudioComponentModifier_SetAttenuation() = default;
+
+		virtual void Execute(UAudioComponent& AudioComponent) const override;
+
+		UPROPERTY(EditAnywhere, Category = "Subsonic")
+		TObjectPtr<USoundAttenuation> Attenuation;
+	};
+
+	USTRUCT(MinimalAPI, DisplayName = "Set Concurrency")
+	struct FSubsonicEventAction_AudioComponentModifier_SetConcurrency : public FSubsonicEventAction_AudioComponentModifierBase
+	{
+		GENERATED_BODY()
+
+	public:
+		virtual ~FSubsonicEventAction_AudioComponentModifier_SetConcurrency() = default;
+
+		virtual void Execute(UAudioComponent& AudioComponent) const override;
+
+		UPROPERTY(EditAnywhere, Category = "Subsonic")
+		TSet<TObjectPtr<USoundConcurrency>> Concurrency;
+	};
+
+	USTRUCT(MinimalAPI, DisplayName = "Set Modulation Routing")
+	struct FSubsonicEventAction_AudioComponentModifier_SetModulationRouting: public FSubsonicEventAction_AudioComponentModifierBase
+	{
+		GENERATED_BODY()
+
+	public:
+		virtual ~FSubsonicEventAction_AudioComponentModifier_SetModulationRouting() = default;
+
+		virtual void Execute(UAudioComponent& AudioComponent) const override;
+
+		UPROPERTY(EditAnywhere, Category = "Subsonic")
+		TSet<TObjectPtr<USoundModulatorBase>> Modulators;
+
+		UPROPERTY(EditAnywhere, Category = "Subsonic")
+		EModulationDestination Destination = EModulationDestination::Volume;
+	
+		UPROPERTY(EditAnywhere, Category = "Subsonic")
+		EModulationRouting RoutingMethod = EModulationRouting::Inherit;
+	};
+
+	USTRUCT(MinimalAPI, DisplayName = "Set Sound")
+	struct FSubsonicEventAction_AudioComponentModifier_SetSound : public FSubsonicEventAction_AudioComponentModifierBase
+	{
+		GENERATED_BODY()
+
+	public:
+		virtual ~FSubsonicEventAction_AudioComponentModifier_SetSound() = default;
+
+		virtual void Execute(UAudioComponent& AudioComponent) const override;
+
+		// Sound to play.
+		UPROPERTY(EditAnywhere, Category = "Subsonic")
+		TObjectPtr<USoundBase> Sound;
+	};
+
+	USTRUCT(MinimalAPI, DisplayName = "Stop")
+	struct FSubsonicEventAction_AudioComponentModifier_Stop : public FSubsonicEventAction_AudioComponentModifierBase
+	{
+		GENERATED_BODY()
+
+	public:
+		virtual ~FSubsonicEventAction_AudioComponentModifier_Stop() = default;
+
+		virtual void Execute(UAudioComponent& AudioComponent) const override;
+	};
+} // namespace UE::Subsonic
\ No newline at end of file
diff --git a/Engine/Plugins/Experimental/Subsonic/Source/SubsonicEngine/Private/StandardActions/SubsonicAction_EventCore.cpp b/Engine/Plugins/Experimental/Subsonic/Source/SubsonicEngine/Private/StandardActions/SubsonicAction_EventCore.cpp
new file mode 100644
--- /dev/null
+++ b/Engine/Plugins/Experimental/Subsonic/Source/SubsonicEngine/Private/StandardActions/SubsonicAction_EventCore.cpp
@@ -0,0 +1,27 @@
+// Copyright Epic Games, Inc. All Rights Reserved.
+
+#include "StandardActions/SubsonicAction_EventCore.h"
+
+#include "SubsonicExecutor.h"
+
+
+namespace UE::Subsonic
+{
+	void FSubsonicEventAction_DelayEvent::Execute(const Core::FSubsonicExecutor& InExecutor, const Core::FActionHandle& InHandle) const
+	{
+		if (EventName.IsValid())
+		{
+			if (Delay > 0.0f)
+			{
+				InExecutor.ExecuteEvent(EventName.GetTagName());
+			}
+			else
+			{
+				// TODO: Implement delay instance state tracking in subscriber & cycle detection
+				// Detection likely should be implemented in validation step required to execute
+				// in non-runtime contexts for perf reasons.
+				InExecutor.ExecuteEvent(EventName.GetTagName());
+			}
+		}
+	}
+} // namespace UE::Subsonic
\ No newline at end of file
diff --git a/Engine/Plugins/Experimental/Subsonic/Source/SubsonicEngine/Private/StandardActions/SubsonicAction_EventCore.h b/Engine/Plugins/Experimental/Subsonic/Source/SubsonicEngine/Private/StandardActions/SubsonicAction_EventCore.h
new file mode 100644
--- /dev/null
+++ b/Engine/Plugins/Experimental/Subsonic/Source/SubsonicEngine/Private/StandardActions/SubsonicAction_EventCore.h
@@ -0,0 +1,49 @@
+// Copyright Epic Games, Inc. All Rights Reserved.
+#pragma once
+
+#include "GameplayTagContainer.h"
+#include "SubsonicEventCollection.h"
+
+#include "SubsonicAction_EventCore.generated.h"
+
+
+namespace UE::Subsonic
+{
+	namespace Core
+	{
+		// Forward Declarations
+		struct FSubsonicExecutor;
+	} // namespace Core
+
+	UENUM()
+	enum class EEventResolutionRule : uint8
+	{
+		First,
+		Last,
+		Random,
+		Index
+	};
+
+	USTRUCT(MinimalAPI, BlueprintType, DisplayName = "Delay Event")
+	struct FSubsonicEventAction_DelayEvent : public Core::FSubsonicEventActionBase
+	{
+		GENERATED_BODY()
+
+	public:
+		virtual ~FSubsonicEventAction_DelayEvent() = default;
+
+		virtual void Execute(const Core::FSubsonicExecutor& InExecutor, const Core::FActionHandle& InHandle) const override;
+
+		// Identifier used to cancel or replace delay in progress.
+		UPROPERTY(EditAnywhere, BlueprintReadWrite, Category = "Action")
+		FName DelayName = "Delay01";
+
+		// Event to execute after the given delay.
+		UPROPERTY(EditAnywhere, BlueprintReadWrite, Category = "Action")
+		FGameplayTag EventName;
+
+		// Executes the given event by the prescribed amount.
+		UPROPERTY(EditAnywhere, BlueprintReadWrite, Category = "Action")
+		float Delay = 1.0f;
+	};
+} // namespace UE::Subsonic
\ No newline at end of file
diff --git a/Engine/Plugins/Experimental/Subsonic/Source/SubsonicEngine/Private/StandardActions/SubsonicAction_GeneratorSource.cpp b/Engine/Plugins/Experimental/Subsonic/Source/SubsonicEngine/Private/StandardActions/SubsonicAction_GeneratorSource.cpp
new file mode 100644
--- /dev/null
+++ b/Engine/Plugins/Experimental/Subsonic/Source/SubsonicEngine/Private/StandardActions/SubsonicAction_GeneratorSource.cpp
@@ -0,0 +1,84 @@
+// Copyright Epic Games, Inc. All Rights Reserved.
+
+#include "StandardActions/SubsonicAction_GeneratorSource.h"
+
+#include "IAudioMixerGeneratorSource.h"
+#include "StandardEventSubscribers/SubsonicGeneratorSourceSubscriber.h"
+#include "SubsonicDeviceUtils.h"
+#include "SubsonicExecutor.h"
+#include "SubsonicHandles.h"
+#include "SubsonicParameterStore.h"
+
+namespace UE::Subsonic
+{
+	void FSubsonicEventAction_GeneratorSourcePlay::Execute(const Core::FSubsonicExecutor& InExecutor, const Core::FActionHandle& InHandle) const
+	{
+		using namespace UE::Subsonic::Core;
+
+		if (Sound == nullptr)
+		{
+			return;
+		}
+
+		// Merge authored params (lowest priority) with trigger-time params from the executor.
+		FSubsonicParameterStore MergedParams = Parameters;
+		MergedParams.MergeFrom(InExecutor.GetParameters());
+
+		const FActionScopeKey ScopeKey(InExecutor, InHandle);
+		if (USubsonicGeneratorSourceSubscriber* GeneratorSourceSubscriber = FindDeviceSubsystem<USubsonicGeneratorSourceSubscriber>(ScopeKey))
+		{
+			switch (Scope)
+			{
+			case ESubsonicExecutionScope::Global:
+				GeneratorSourceSubscriber->PlaySound(Name, *Sound, MergedParams);
+				break;
+
+			case ESubsonicExecutionScope::Collection:
+				GeneratorSourceSubscriber->PlaySound(InHandle.Event.Collection, Name, *Sound, MergedParams);
+				break;
+
+			case ESubsonicExecutionScope::Event:
+				GeneratorSourceSubscriber->PlaySound(InHandle.Event, Name, *Sound, MergedParams);
+				break;
+
+			case ESubsonicExecutionScope::Action:
+				GeneratorSourceSubscriber->PlaySound(InHandle, Name, *Sound, MergedParams);
+				break;
+
+			default:
+				checkNoEntry();
+			}
+		}
+	}
+
+	void FSubsonicEventAction_GeneratorSourceStop::Execute(const Core::FSubsonicExecutor& InExecutor, const Core::FActionHandle& InHandle) const
+	{
+		using namespace UE::Subsonic::Core;
+
+		const FActionScopeKey ScopeKey(InExecutor, InHandle);
+		if (USubsonicGeneratorSourceSubscriber* GeneratorSourceSubscriber = FindDeviceSubsystem<USubsonicGeneratorSourceSubscriber>(ScopeKey))
+		{
+			switch (Scope)
+			{
+			case ESubsonicExecutionScope::Global:
+				GeneratorSourceSubscriber->StopSound(Name);
+				break;
+
+			case ESubsonicExecutionScope::Collection:
+				GeneratorSourceSubscriber->StopSound(InHandle.Event.Collection, Name);
+				break;
+
+			case ESubsonicExecutionScope::Event:
+				GeneratorSourceSubscriber->StopSound(InHandle.Event, Name);
+				break;
+
+			case ESubsonicExecutionScope::Action:
+				GeneratorSourceSubscriber->StopSound(InHandle, Name);
+				break;
+
+			default:
+				checkNoEntry();
+			}
+		}
+	}
+} // namespace UE::Subsonic
\ No newline at end of file
diff --git a/Engine/Plugins/Experimental/Subsonic/Source/SubsonicEngine/Private/StandardEventSubscribers/AudioComponentEventSubscriber.cpp b/Engine/Plugins/Experimental/Subsonic/Source/SubsonicEngine/Private/StandardEventSubscribers/AudioComponentEventSubscriber.cpp
new file mode 100644
--- /dev/null
+++ b/Engine/Plugins/Experimental/Subsonic/Source/SubsonicEngine/Private/StandardEventSubscribers/AudioComponentEventSubscriber.cpp
@@ -0,0 +1,567 @@
+// Copyright Epic Games, Inc. All Rights Reserved.
+
+#include "StandardEventSubscribers/AudioComponentEventSubscriber.h"
+
+#include "Algo/Transform.h"
+#include "AudioDevice.h"
+#include "AudioDeviceManager.h"
+#include "Engine/Engine.h"
+#include "SubsonicCoreLog.h"
+#include "SubsonicExecutor.h"
+#include "SubsonicHandles.h"
+#include "SubsonicSubsystem.h"
+#include "UObject/GCObject.h"
+#include "UObject/UObjectGlobals.h"
+
+#define SUBSONIC_ASSERT_ON_STALE_COLLECTION_EXECUTORS (1 && DO_CHECK)
+
+namespace UE::Subsonic
+{
+	namespace AudioComponentEventSubscriberPrivate
+	{
+		struct FInstanceAudioComponents : public FGCObject
+		{
+			// Currently addressable component via handle
+			TObjectPtr<UAudioComponent> ActiveComponent;
+
+			// Components no longer addressable from actions but actively
+			// playing. Continually tracked here to avoid GC while finishing.
+			TArray<TObjectPtr<UAudioComponent>> ReleasedComponents;
+
+			virtual void AddReferencedObjects(FReferenceCollector& Collector) override
+			{
+				Collector.AddReferencedObject(ActiveComponent);
+				Collector.AddStableReferenceArray(&ReleasedComponents);
+			}
+
+			virtual FString GetReferencerName() const override
+			{
+				return "UE::Subsonic::InstanceAudioComponents";
+			}
+
+			bool IsEmpty() const
+			{
+				return !ActiveComponent.Get() && ReleasedComponents.IsEmpty();
+			}
+
+			void ReleaseComponent()
+			{
+				if (ActiveComponent.Get() && ActiveComponent->IsPlaying())
+				{
+					ReleasedComponents.Add(ActiveComponent);
+				}
+				ActiveComponent = nullptr;
+			}
+
+			void StopAll() const
+			{
+				if (UAudioComponent* Component = ActiveComponent.Get())
+				{
+					Component->Stop();
+				}
+
+				for (const TObjectPtr<UAudioComponent>& ReleasedComponent : ReleasedComponents)
+				{
+					if (UAudioComponent* Component = ReleasedComponent.Get())
+					{
+						Component->Stop();
+					}
+				}
+			}
+		};
+
+		using FInstanceAudioComponentsPtr = TUniquePtr<FInstanceAudioComponents>;
+		using FAudioComponentsCache = TMap<FName, FInstanceAudioComponentsPtr>;
+
+		class FSubscriberImpl : public TSharedFromThis<FSubscriberImpl>
+		{
+			FAudioComponentsCache GlobalComponents;
+
+			TMap<Core::FCollectionHandle, FAudioComponentsCache> CollectionComponents;
+			TMap<Core::FEventHandle, FAudioComponentsCache> EventComponents;
+			TMap<Core::FActionHandle, FAudioComponentsCache> ActionComponents;
+
+			TMap<uint32, Core::FCollectionHandle> CollectionExecutors;
+
+			static UAudioComponent* CreateComponent(UObject& Parent, Audio::FDeviceId DeviceId, FName Name)
+			{
+				const FName ComponentName = MakeUniqueObjectName(&Parent, UAudioComponent::StaticClass(), Name);
+				UAudioComponent* NewComponent = NewObject<UAudioComponent>(&Parent, ComponentName);
+
+				check(NewComponent);
+				NewComponent->AudioDeviceID = DeviceId;
+				NewComponent->bAutoDestroy = true;
+				NewComponent->bIsUISound = true;
+				NewComponent->bAllowSpatialization = false;
+				NewComponent->bReverb = false;
+				NewComponent->bCenterChannelOnly = false;
+
+				return NewComponent;
+			}
+
+			static UAudioComponent* AddCachedComponent(
+				UObject& Parent,
+				Audio::FDeviceId DeviceId,
+				const FInstanceAudioComponentsPtr& InstanceComponents,
+				FName Name,
+				bool bReleaseExisting,
+				bool& bOutStoppedComponent)
+			{
+				bOutStoppedComponent = false;
+
+				if (bReleaseExisting)
+				{
+					if (UAudioComponent* Component = InstanceComponents->ActiveComponent.Get(); Component && Component->IsPlaying())
+					{
+						if (USoundBase* Sound = Component->GetSound(); Sound && !Sound->IsOneShot())
+						{
+							bOutStoppedComponent = true;
+							Component->Stop();
+						}
+					}
+
+					InstanceComponents->ReleaseComponent();
+				}
+
+				TObjectPtr<UAudioComponent>& ComponentRef = InstanceComponents->ActiveComponent;
+				if (!ComponentRef)
+				{
+					ComponentRef = CreateComponent(Parent, DeviceId, Name);
+				}
+
+				return ComponentRef.Get();
+			}
+
+			template <typename HandleType>
+			UAudioComponent* AddCachedComponent(
+				UObject& Parent,
+				Audio::FDeviceId DeviceId,
+				const HandleType& InHandle,
+				TMap<HandleType, FAudioComponentsCache>& OutScopeComponentMap,
+				FName Name,
+				bool bReleaseExisting)
+			{
+				UAudioComponent* CacheComponent = nullptr;
+
+				FAudioComponentsCache& Cache = OutScopeComponentMap.FindOrAdd(InHandle);
+				FInstanceAudioComponentsPtr& InstanceComponents = Cache.FindOrAdd(Name);
+				if (!InstanceComponents.IsValid())
+				{
+					InstanceComponents = MakeUnique<FInstanceAudioComponents>();
+				}
+
+				bool bStoppedComponent = false;
+				CacheComponent = AddCachedComponent(Parent, DeviceId, InstanceComponents, Name, bReleaseExisting, bStoppedComponent);
+				if (bStoppedComponent)
+				{
+					UE_LOG(LogSubsonic, Warning,
+						TEXT("Stopping released sound with handle '%s': "
+							"Sound is not one shot and continued play would result in leaked playback."),
+						*InHandle.ToString())
+				}
+				auto RemoveComponentLambda = [FinishHandle = InHandle, &OutScopeComponentMap, Name](UAudioComponent* FinishingComponent)
+				{
+					if (FAudioComponentsCache* Cache = OutScopeComponentMap.Find(FinishHandle))
+					{
+						if (FInstanceAudioComponentsPtr* CacheComponentsPtr = Cache->Find(Name); CacheComponentsPtr && CacheComponentsPtr->IsValid())
+						{
+							FInstanceAudioComponentsPtr& CacheComponents = *CacheComponentsPtr;
+							if (CacheComponents->ActiveComponent.Get() == FinishingComponent)
+							{
+								CacheComponents->ActiveComponent = nullptr;
+							}
+
+							CacheComponents->ReleasedComponents.RemoveAllSwap([&FinishingComponent](const TObjectPtr<UAudioComponent>& Released)
+							{
+								return Released.Get() == FinishingComponent;
+							});
+
+							if (CacheComponents->IsEmpty())
+							{
+								Cache->Remove(Name);
+							}
+						}
+						if (Cache->IsEmpty())
+						{
+							OutScopeComponentMap.Remove(FinishHandle);
+						}
+					}
+				};
+
+				CacheComponent->OnAudioFinishedNative.AddSPLambda(this, RemoveComponentLambda);
+				return CacheComponent;
+			}
+
+			template <typename HandleType>
+			UAudioComponent* FindCacheComponent(const TMap<HandleType, FAudioComponentsCache>& ScopeComponentMap, const HandleType& InHandle, FName Name) const
+			{
+				if (const FAudioComponentsCache* Cache = ScopeComponentMap.Find(InHandle))
+				{
+					if (const FInstanceAudioComponentsPtr* CacheComponentsPtr = Cache->Find(Name); CacheComponentsPtr && CacheComponentsPtr->IsValid())
+					{
+						return (*CacheComponentsPtr)->ActiveComponent.Get();
+					}
+				}
+
+				return nullptr;
+			}
+
+			template <typename HandleType>
+			bool ContainsCacheComponent(const TMap<HandleType, FAudioComponentsCache>& ScopeComponentMap, const HandleType& InHandle, FName Name) const
+			{
+				if (const FAudioComponentsCache* Cache = ScopeComponentMap.Find(InHandle))
+				{
+					if (const FInstanceAudioComponentsPtr* CacheComponentsPtr = Cache->Find(Name))
+					{
+						return CacheComponentsPtr->IsValid();
+					}
+				}
+
+				return false;
+			}
+
+			template <typename HandleType>
+			void RemoveCacheComponent(TMap<HandleType, FAudioComponentsCache>& ScopeComponentMap, const HandleType& InHandle, FName Name)
+			{
+				if (FAudioComponentsCache* Cache = ScopeComponentMap.Find(InHandle))
+				{
+					FInstanceAudioComponentsPtr* InstanceComponents = Cache->Find(Name);
+					if (InstanceComponents && InstanceComponents->IsValid())
+					{
+						(*InstanceComponents)->StopAll();
+					}
+
+					Cache->Remove(Name);
+					if (Cache->IsEmpty())
+					{
+						ScopeComponentMap.Remove(InHandle);
+					}
+				}
+			}
+
+			template <typename THandleType>
+			void ResetScopeComponents(TMap<THandleType, FAudioComponentsCache>& ScopeComponentMap)
+			{
+				for (const TPair<THandleType, FAudioComponentsCache>& HandleCachePair : ScopeComponentMap)
+				{
+					for (const TPair<FName, FInstanceAudioComponentsPtr>& Pair : HandleCachePair.Value)
+					{
+						const FInstanceAudioComponentsPtr& InstanceComponents = Pair.Value;
+						check(InstanceComponents.IsValid());
+						InstanceComponents->StopAll();
+					}
+				}
+				ScopeComponentMap.Reset();
+			};
+
+		public:
+			void Deinitialize()
+			{
+				for (const TPair<FName, FInstanceAudioComponentsPtr>& Pair : GlobalComponents)
+				{
+					const FInstanceAudioComponentsPtr& InstanceComponents = Pair.Value;
+					check(InstanceComponents.IsValid());
+					InstanceComponents->StopAll();
+				}
+				GlobalComponents.Empty();
+
+				ResetScopeComponents(CollectionComponents);
+				ResetScopeComponents(EventComponents);
+				ResetScopeComponents(ActionComponents);
+			}
+
+			UAudioComponent* AddComponent(UObject& Parent, Audio::FDeviceId DeviceId, FName Name, bool bReleaseExisting)
+			{
+				FInstanceAudioComponentsPtr& InstanceComponents = GlobalComponents.FindOrAdd(Name);
+				if (!InstanceComponents.IsValid())
+				{
+					InstanceComponents = MakeUnique<FInstanceAudioComponents>();
+				}
+
+				bool bStoppedComponent = false;
+				if (UAudioComponent* CacheComponent = AddCachedComponent(Parent, DeviceId, InstanceComponents, Name, bReleaseExisting, bStoppedComponent))
+				{
+					if (bStoppedComponent)
+					{
+						UE_LOG(LogSubsonic, Warning,
+							TEXT("Stopping released global sound '%s': "
+								"Sound is not one shot and continued play would result in leaked playback."),
+							*Name.ToString())
+					}
+
+					auto RemoveComponentLambda = [this, Name](UAudioComponent* FinishingComponent)
+					{
+						if (FInstanceAudioComponentsPtr* Components = GlobalComponents.Find(Name); Components && Components->IsValid())
+						{
+							FInstanceAudioComponents& ComponentsRef = (*Components->Get());
+							if (ComponentsRef.ActiveComponent.Get() == FinishingComponent)
+							{
+								ComponentsRef.ActiveComponent = nullptr;
+							}
+
+							ComponentsRef.ReleasedComponents.RemoveAllSwap([&FinishingComponent](const TObjectPtr<UAudioComponent>& Released)
+							{
+								return Released.Get() == FinishingComponent;
+							});
+
+							if (ComponentsRef.IsEmpty())
+							{
+								GlobalComponents.Remove(Name);
+							}
+						}
+					};
+					CacheComponent->OnAudioFinishedNative.AddSPLambda(this, RemoveComponentLambda);
+					return CacheComponent;
+				}
+
+				return nullptr;
+			}
+
+			UAudioComponent* AddComponent(UObject& Parent, Audio::FDeviceId DeviceId, const Core::FCollectionHandle& InHandle, FName Name, bool bReleaseExisting)
+			{
+				return AddCachedComponent(Parent, DeviceId, InHandle, CollectionComponents, Name, bReleaseExisting);
+			}
+
+			UAudioComponent* AddComponent(UObject& Parent, Audio::FDeviceId DeviceId, const Core::FEventHandle& InHandle, FName Name, bool bReleaseExisting)
+			{
+				return AddCachedComponent(Parent, DeviceId, InHandle, EventComponents, Name, bReleaseExisting);
+			}
+
+			UAudioComponent* AddComponent(UObject& Parent, Audio::FDeviceId DeviceId, const Core::FActionHandle& InHandle, FName Name, bool bReleaseExisting)
+			{
+				return AddCachedComponent(Parent, DeviceId, InHandle, ActionComponents, Name, bReleaseExisting);
+			}
+
+			bool ContainsComponent(FName Name) const
+			{
+				return GlobalComponents.Contains(Name);
+			}
+
+			bool ContainsComponent(const Core::FCollectionHandle& InHandle, FName Name) const
+			{
+				return ContainsCacheComponent(CollectionComponents, InHandle, Name);
+			}
+
+			bool ContainsComponent(const Core::FEventHandle& InHandle, FName Name) const
+			{
+				return ContainsCacheComponent(EventComponents, InHandle, Name);
+			}
+
+			bool ContainsComponent(const Core::FActionHandle& InHandle, FName Name) const
+			{
+				return ContainsCacheComponent(ActionComponents, InHandle, Name);
+			}
+
+			UAudioComponent* FindComponent(FName Name) const
+			{
+				const FInstanceAudioComponentsPtr* InstanceComponents = GlobalComponents.Find(Name);
+				if (InstanceComponents && InstanceComponents->IsValid())
+				{
+					return (*InstanceComponents)->ActiveComponent.Get();
+				}
+
+				return nullptr;
+			}
+
+			UAudioComponent* FindComponent(const Core::FCollectionHandle& InHandle, FName Name) const
+			{
+				return FindCacheComponent(CollectionComponents, InHandle, Name);
+			}
+
+			UAudioComponent* FindComponent(const Core::FEventHandle& InHandle, FName Name) const
+			{
+				return FindCacheComponent(EventComponents, InHandle, Name);
+			}
+
+			UAudioComponent* FindComponent(const Core::FActionHandle& InHandle, FName Name) const
+			{
+				return FindCacheComponent(ActionComponents, InHandle, Name);
+			}
+
+			void OnCollectionUnregistered(const Core::FCollectionHandle& InHandle)
+			{
+				if (FAudioComponentsCache* Cache = CollectionComponents.Find(InHandle))
+				{
+					for (TPair<FName, FInstanceAudioComponentsPtr>& Pair : *Cache)
+					{
+						FInstanceAudioComponentsPtr& InstanceComponents = Pair.Value;
+						check(InstanceComponents.IsValid());
+						if (UAudioComponent* Component = InstanceComponents->ActiveComponent.Get())
+						{
+							Component->Stop();
+						}
+
+						for (const TObjectPtr<UAudioComponent>& ReleasedComponent : InstanceComponents->ReleasedComponents)
+						{
+							if (UAudioComponent* Component = ReleasedComponent.Get())
+							{
+								Component->Stop();
+							}
+						}
+					}
+
+					CollectionComponents.Remove(InHandle);
+
+	#if SUBSONIC_ASSERT_ON_STALE_COLLECTION_EXECUTORS
+					for (TPair<uint32, Core::FCollectionHandle>& Pair : CollectionExecutors)
+					{
+						checkf(Pair.Value.CollectionId != InHandle.CollectionId, TEXT("Executor for Collection '%s' still exists beyond lifetime of collection."), *InHandle.ToString());
+					}
+	#endif // SUBSONIC_ASSERT_ON_STALE_COLLECTION_EXECUTORS
+				}
+			}
+
+			void OnExecutorRegistered(const Core::FSubsonicExecutor& InExecutor)
+			{
+				CollectionExecutors.Add(InExecutor.GetId(), InExecutor.GetCollectionHandle());
+			}
+
+			void OnExecutorUnregistered(const Core::FSubsonicExecutor& InExecutor)
+			{
+				if (const Core::FCollectionHandle* CollectionHandle = CollectionExecutors.Find(InExecutor.GetId()))
+				{
+					if (const FAudioComponentsCache* ScopeComponentMap = CollectionComponents.Find(*CollectionHandle))
+					{
+						for (const TPair<FName, FInstanceAudioComponentsPtr>& Pair : *ScopeComponentMap)
+						{
+							if (FInstanceAudioComponents* Components = Pair.Value.Get())
+							{
+								Components->StopAll();
+							}
+						}
+					}
+					CollectionExecutors.Remove(InExecutor.GetId());
+					CollectionComponents.Remove(*CollectionHandle);
+				}
+			}
+
+			void RemoveComponent(FName Name)
+			{
+				FInstanceAudioComponentsPtr* InstanceComponents = GlobalComponents.Find(Name);
+				if (InstanceComponents && InstanceComponents->IsValid())
+				{
+					(*InstanceComponents)->StopAll();
+					GlobalComponents.Remove(Name);
+				}
+			}
+
+			void RemoveComponent(const Core::FCollectionHandle& InHandle, FName Name)
+			{
+				return RemoveCacheComponent(CollectionComponents, InHandle, Name);
+			}
+
+			void RemoveComponent(const Core::FEventHandle& InHandle, FName Name)
+			{
+				return RemoveCacheComponent(EventComponents, InHandle, Name);
+			}
+
+			void RemoveComponent(const Core::FActionHandle& InHandle, FName Name)
+			{
+				return RemoveCacheComponent(ActionComponents, InHandle, Name);
+			}
+		};
+	} // namespace AudioComponentEventSubscriberPrivate
+
+	void USubsonicAudioComponentSubscriber::Initialize(FSubsystemCollectionBase& Collection)
+	{
+		Super::Initialize(Collection);
+		Impl = MakeShared<AudioComponentEventSubscriberPrivate::FSubscriberImpl>();
+	}
+
+	void USubsonicAudioComponentSubscriber::Deinitialize()
+	{
+		Impl->Deinitialize();
+		Super::Deinitialize();
+	}
+
+	UAudioComponent* USubsonicAudioComponentSubscriber::AddComponent(FName Name, bool bReleaseExisting)
+	{
+		return Impl->AddComponent(*this, GetAudioDeviceId(), Name, bReleaseExisting);
+	}
+
+	UAudioComponent* USubsonicAudioComponentSubscriber::AddComponent(const Core::FCollectionHandle& InHandle, FName Name, bool bReleaseExisting)
+	{
+		return Impl->AddComponent(*this, GetAudioDeviceId(), InHandle, Name, bReleaseExisting);
+	}
+
+	UAudioComponent* USubsonicAudioComponentSubscriber::AddComponent(const Core::FEventHandle& InHandle, FName Name, bool bReleaseExisting)
+	{
+		return Impl->AddComponent(*this, GetAudioDeviceId(), InHandle, Name, bReleaseExisting);
+	}
+
+	UAudioComponent* USubsonicAudioComponentSubscriber::AddComponent(const Core::FActionHandle& InHandle, FName Name, bool bReleaseExisting)
+	{
+		return Impl->AddComponent(*this, GetAudioDeviceId(), InHandle, Name, bReleaseExisting);
+	}
+
+	bool USubsonicAudioComponentSubscriber::ContainsComponent(FName Name) const
+	{
+		return Impl->ContainsComponent(Name);
+	}
+
+	bool USubsonicAudioComponentSubscriber::ContainsComponent(const Core::FCollectionHandle& InHandle, FName Name) const
+	{
+		return Impl->ContainsComponent(InHandle, Name);
+	}
+
+	bool USubsonicAudioComponentSubscriber::ContainsComponent(const Core::FEventHandle& InHandle, FName Name) const
+	{
+		return Impl->ContainsComponent(InHandle, Name);
+	}
+
+	bool USubsonicAudioComponentSubscriber::ContainsComponent(const Core::FActionHandle& InHandle, FName Name) const
+	{
+		return Impl->ContainsComponent(InHandle, Name);
+	}
+
+	UAudioComponent* USubsonicAudioComponentSubscriber::FindComponent(FName Name)
+	{
+		return Impl->FindComponent(Name);
+	}
+
+	UAudioComponent* USubsonicAudioComponentSubscriber::FindComponent(const Core::FCollectionHandle& InHandle, FName Name)
+	{
+		return Impl->FindComponent(InHandle, Name);
+	}
+
+	UAudioComponent* USubsonicAudioComponentSubscriber::FindComponent(const Core::FEventHandle& InHandle, FName Name)
+	{
+		return Impl->FindComponent(InHandle, Name);
+	}
+
+	UAudioComponent* USubsonicAudioComponentSubscriber::FindComponent(const Core::FActionHandle& InHandle, FName Name)
+	{
+		return Impl->FindComponent(InHandle, Name);
+	}
+
+	void USubsonicAudioComponentSubscriber::RemoveComponent(FName Name)
+	{
+		return Impl->RemoveComponent(Name);
+	}
+
+	void USubsonicAudioComponentSubscriber::RemoveComponent(const Core::FCollectionHandle& InHandle, FName Name)
+	{
+		return Impl->RemoveComponent(InHandle, Name);
+	}
+
+	void USubsonicAudioComponentSubscriber::RemoveComponent(const Core::FEventHandle& InHandle, FName Name)
+	{
+		return Impl->RemoveComponent(InHandle, Name);
+	}
+
+	void USubsonicAudioComponentSubscriber::OnCollectionUnregistered(const Core::FCollectionHandle& InHandle)
+	{
+		Impl->OnCollectionUnregistered(InHandle);
+	}
+
+	void USubsonicAudioComponentSubscriber::OnExecutorRegistered(const Core::FSubsonicExecutor& InExecutor)
+	{
+		Impl->OnExecutorRegistered(InExecutor);
+	}
+
+	void USubsonicAudioComponentSubscriber::OnExecutorUnregistered(const Core::FSubsonicExecutor& InExecutor)
+	{
+		Impl->OnExecutorUnregistered(InExecutor);
+	}
+} // namespace UE::Subsonic
+#undef SUBSONIC_ASSERT_ON_STALE_COLLECTION_EXECUTORS
diff --git a/Engine/Plugins/Experimental/Subsonic/Source/SubsonicEngine/Private/StandardEventSubscribers/AudioComponentEventSubscriber.h b/Engine/Plugins/Experimental/Subsonic/Source/SubsonicEngine/Private/StandardEventSubscribers/AudioComponentEventSubscriber.h
new file mode 100644
--- /dev/null
+++ b/Engine/Plugins/Experimental/Subsonic/Source/SubsonicEngine/Private/StandardEventSubscribers/AudioComponentEventSubscriber.h
@@ -0,0 +1,71 @@
+// Copyright Epic Games, Inc. All Rights Reserved.
+#pragma once
+
+#include "Subsystems/EngineSubsystem.h"
+
+#include "Components/AudioComponent.h"
+#include "Containers/Map.h"
+#include "StandardEventSubscribers/SubsonicEventSubscriberBase.h"
+#include "SubsonicEventSubscriberInterface.h"
+#include "SubsonicHandles.h"
+#include "Subsystems/AudioEngineSubsystem.h"
+#include "Templates/SharedPointer.h"
+
+#include "AudioComponentEventSubscriber.generated.h"
+
+#define UE_API SUBSONICENGINE_API
+
+
+// Forward Declarations
+namespace UE::Subsonic
+{
+	struct FSubsonicExecutor;
+
+	namespace AudioComponentEventSubscriberPrivate
+	{
+		class FSubscriberImpl;
+	} // namespace AudioComponentEventSubscriberPrivate
+
+	/**
+	 * Subsonic Subscriber managing AudioComponent instances, events, and related data.
+	 */
+	UCLASS(MinimalAPI)
+	class USubsonicAudioComponentSubscriber final : public USubsonicEventSubscriberBase
+	{
+		GENERATED_BODY()
+
+	public:
+		//~ Begin USubsonicEventSubscriberBase interface
+		UE_API virtual void Initialize(FSubsystemCollectionBase& Collection) override;
+		UE_API virtual void Deinitialize() override;
+		//~ End USubsonicEventSubscriberBase interface
+
+		UE_INTERNAL virtual void OnCollectionUnregistered(const Core::FCollectionHandle& InEvent) override;
+		UE_INTERNAL virtual void OnExecutorRegistered(const Core::FSubsonicExecutor& InExecutor) override;
+		UE_INTERNAL virtual void OnExecutorUnregistered(const Core::FSubsonicExecutor& InExecutor) override;
+
+		UE_INTERNAL UAudioComponent* AddComponent(FName Name, bool bReleaseExisting = true);
+		UE_INTERNAL UAudioComponent* AddComponent(const Core::FCollectionHandle& InHandle, FName Name, bool bReleaseExisting = true);
+		UE_INTERNAL UAudioComponent* AddComponent(const Core::FEventHandle& InHandle, FName Name, bool bReleaseExisting = true);
+		UE_INTERNAL UAudioComponent* AddComponent(const Core::FActionHandle& InHandle, FName Name, bool bReleaseExisting = true);
+
+		UE_INTERNAL bool ContainsComponent(FName Name) const;
+		UE_INTERNAL bool ContainsComponent(const Core::FCollectionHandle& InHandle, FName Name) const;
+		UE_INTERNAL bool ContainsComponent(const Core::FEventHandle& InHandle, FName Name) const;
+		UE_INTERNAL bool ContainsComponent(const Core::FActionHandle& InHandle, FName Name) const;
+
+		UE_INTERNAL UAudioComponent* FindComponent(FName Name);
+		UE_INTERNAL UAudioComponent* FindComponent(const Core::FCollectionHandle& InHandle, FName Name);
+		UE_INTERNAL UAudioComponent* FindComponent(const Core::FEventHandle& InHandle, FName Name);
+		UE_INTERNAL UAudioComponent* FindComponent(const Core::FActionHandle& InHandle, FName Name);
+
+		UE_INTERNAL void RemoveComponent(FName Name);
+		UE_INTERNAL void RemoveComponent(const Core::FCollectionHandle& InHandle, FName Name);
+		UE_INTERNAL void RemoveComponent(const Core::FEventHandle& InHandle, FName Name);
+
+	private:
+		TSharedPtr<AudioComponentEventSubscriberPrivate::FSubscriberImpl> Impl;
+	};
+} // namespace UE::Subsonic
+
+#undef UE_API // SUBSONICENGINE_API
diff --git a/Engine/Plugins/Experimental/Subsonic/Source/SubsonicEngine/Private/StandardEventSubscribers/SubsonicEventSubscriberBase.cpp b/Engine/Plugins/Experimental/Subsonic/Source/SubsonicEngine/Private/StandardEventSubscribers/SubsonicEventSubscriberBase.cpp
new file mode 100644
--- /dev/null
+++ b/Engine/Plugins/Experimental/Subsonic/Source/SubsonicEngine/Private/StandardEventSubscribers/SubsonicEventSubscriberBase.cpp
@@ -0,0 +1,16 @@
+// Copyright Epic Games, Inc. All Rights Reserved.
+
+#include "StandardEventSubscribers/SubsonicEventSubscriberBase.h"
+
+namespace UE::Subsonic
+{
+	void USubsonicEventSubscriberBase::Initialize(FSubsystemCollectionBase& Collection)
+	{
+		Core::ISubsonicEventSubscriberInterface::Register();
+	}
+
+	void USubsonicEventSubscriberBase::Deinitialize()
+	{
+		Core::ISubsonicEventSubscriberInterface::Unregister();
+	}
+} // namespace UE::Subsonic
\ No newline at end of file
diff --git a/Engine/Plugins/Experimental/Subsonic/Source/SubsonicEngine/Private/StandardEventSubscribers/SubsonicGenerator.cpp b/Engine/Plugins/Experimental/Subsonic/Source/SubsonicEngine/Private/StandardEventSubscribers/SubsonicGenerator.cpp
new file mode 100644
--- /dev/null
+++ b/Engine/Plugins/Experimental/Subsonic/Source/SubsonicEngine/Private/StandardEventSubscribers/SubsonicGenerator.cpp
@@ -0,0 +1,407 @@
+// Copyright Epic Games, Inc. All Rights Reserved.
+
+#include "SubsonicGenerator.h"
+
+#include "DSP/Dsp.h"
+#include "DSP/FloatArrayMath.h"
+#include "IAudioMixerGeneratorSource.h"
+#include "StandardEventSubscribers/SubsonicWaveGenerator.h"
+#include "StructUtils/PropertyBag.h"
+#include "SubsonicBuiltInParameters.h"
+#include "SubsonicCoreLog.h"
+#include "SubsonicParameterStore.h"
+
+namespace UE::Subsonic
+{
+	const TMap<FName, FSubsonicGenerator::FBuiltInParamHandler>& FSubsonicGenerator::GetBuiltInParams()
+	{
+		static const TMap<FName, FBuiltInParamHandler> Map = {
+			{ BuiltInParameters::Volume,         &FSubsonicGenerator::SetVolume },
+			{ BuiltInParameters::PitchShift,     &FSubsonicGenerator::SetPitchShift },
+			{ BuiltInParameters::HighpassCutoff, &FSubsonicGenerator::SetHighpassCutoff },
+			{ BuiltInParameters::LowpassCutoff,  &FSubsonicGenerator::SetLowpassCutoff },
+			{ BuiltInParameters::FadeOutTime,    &FSubsonicGenerator::SetFadeOutTime },
+		};
+		return Map;
+	}
+
+	FSubsonicGenerator::FSubsonicGenerator(TSharedRef<FWaveGenerator> InWave, float InSampleRate)
+		: InnerGenerator(InWave)
+		, bGraphReady(true) // Wave path: no graph to wait for.
+		, SampleRate(InSampleRate)
+	{
+		InitDSP();
+	}
+
+	FSubsonicGenerator::FSubsonicGenerator(ISoundGeneratorPtr InMetaSoundGenerator, float InSampleRate)
+		: InnerGenerator(InMetaSoundGenerator)
+		, bGraphReady(false) // Will be set true by the graph-set callback.
+		, SampleRate(InSampleRate)
+	{
+		TSharedPtr<Metasound::FMetasoundGenerator, ESPMode::ThreadSafe> MSGen =
+			StaticCastSharedPtr<Metasound::FMetasoundGenerator>(InMetaSoundGenerator);
+		if (MSGen.IsValid())
+		{
+			MetaSoundGeneratorWeak = MSGen;
+
+			// Register for notification when the MetaSound graph is compiled and ready.
+			// If the graph is already built, AddGraphSetCallback fires immediately.
+			GraphSetCallbackHandle = MSGen->AddGraphSetCallback(
+				Metasound::FOnSetGraph::FDelegate::CreateLambda([this]()
+				{
+					bGraphReady.store(true, std::memory_order_release);
+				}));
+		}
+		else
+		{
+			bGraphReady.store(true, std::memory_order_relaxed);
+		}
+
+		InitDSP();
+	}
+
+	FSubsonicGenerator::~FSubsonicGenerator()
+	{
+		if (GraphSetCallbackHandle.IsValid())
+		{
+			if (TSharedPtr<Metasound::FMetasoundGenerator, ESPMode::ThreadSafe> MSGen = MetaSoundGeneratorWeak.Pin())
+			{
+				MSGen->RemoveGraphSetCallback(GraphSetCallbackHandle);
+			}
+		}
+	}
+
+	void FSubsonicGenerator::InitDSP()
+	{
+		const int32 NumCh = InnerGenerator->GetNumChannels();
+
+		HighpassFilter.Init(SampleRate, NumCh, Audio::EBiquadFilter::Highpass, 0.0f);
+		HighpassFilter.SetEnabled(false);
+
+		LowpassFilter.Init(SampleRate, NumCh, Audio::EBiquadFilter::Lowpass, SampleRate * 0.5f);
+		LowpassFilter.SetEnabled(false);
+
+		Resampler = MakeUnique<Audio::FRuntimeResampler>(NumCh);
+		Resampler->SetFrameRatio(1.0f);
+	}
+
+	bool FSubsonicGenerator::TryGetFloat(const FSubsonicParameterStore& Store, const FPropertyBagPropertyDesc& Desc, float& OutValue)
+	{
+		if (Desc.ValueType == EPropertyBagPropertyType::Float)
+		{
+			TValueOrError<float, EPropertyBagResult> Result = Store.Bag.GetValueFloat(Desc);
+			if (Result.HasValue())
+			{
+				OutValue = Result.GetValue();
+				return true;
+			}
+		}
+		else if (Desc.ValueType == EPropertyBagPropertyType::Double)
+		{
+			TValueOrError<double, EPropertyBagResult> Result = Store.Bag.GetValueDouble(Desc);
+			if (Result.HasValue())
+			{
+				OutValue = static_cast<float>(Result.GetValue());
+				return true;
+			}
+		}
+		return false;
+	}
+
+	bool FSubsonicGenerator::HasMetaSoundInput(FName Name) const
+	{
+		TSharedPtr<Metasound::FMetasoundGenerator, ESPMode::ThreadSafe> MSGen = MetaSoundGeneratorWeak.Pin();
+		if (!MSGen.IsValid())
+		{
+			return false;
+		}
+		return MSGen->GetInputWriteReference<float>(Name).IsSet();
+	}
+
+	void FSubsonicGenerator::ForwardToMetaSound(
+		Metasound::FMetasoundGenerator& MSGen,
+		const FSubsonicParameterStore& Store,
+		const FPropertyBagPropertyDesc& Desc)
+	{
+		switch (Desc.ValueType)
+		{
+		case EPropertyBagPropertyType::Float:
+		{
+			TValueOrError<float, EPropertyBagResult> Result = Store.Bag.GetValueFloat(Desc);
+			if (Result.HasValue())
+			{
+				MSGen.SetInputValue<float>(Desc.Name, Result.GetValue());
+			}
+			break;
+		}
+		case EPropertyBagPropertyType::Double:
+		{
+			TValueOrError<double, EPropertyBagResult> Result = Store.Bag.GetValueDouble(Desc);
+			if (Result.HasValue())
+			{
+				MSGen.SetInputValue<float>(Desc.Name, static_cast<float>(Result.GetValue()));
+			}
+			break;
+		}
+		case EPropertyBagPropertyType::Bool:
+		{
+			TValueOrError<bool, EPropertyBagResult> Result = Store.Bag.GetValueBool(Desc);
+			if (Result.HasValue())
+			{
+				MSGen.SetInputValue<bool>(Desc.Name, Result.GetValue());
+			}
+			break;
+		}
+		case EPropertyBagPropertyType::Int32:
+		{
+			TValueOrError<int32, EPropertyBagResult> Result = Store.Bag.GetValueInt32(Desc);
+			if (Result.HasValue())
+			{
+				MSGen.SetInputValue<int32>(Desc.Name, Result.GetValue());
+			}
+			break;
+		}
+		default:
+			break;
+		}
+	}
+
+	void FSubsonicGenerator::ApplyParameters(const FSubsonicParameterStore& Store)
+	{
+		const UPropertyBag* BagStruct = Store.Bag.GetPropertyBagStruct();
+		if (!BagStruct)
+		{
+			return;
+		}
+
+		// If the MetaSound graph isn't ready yet, defer the entire store.
+		if (!bGraphReady.load(std::memory_order_acquire))
+		{
+			DeferredParams = MakeShared<FSubsonicParameterStore>(Store);
+			return;
+		}
+
+		TSharedPtr<Metasound::FMetasoundGenerator, ESPMode::ThreadSafe> MSGen = MetaSoundGeneratorWeak.Pin();
+		const TMap<FName, FBuiltInParamHandler>& BuiltInParams = GetBuiltInParams();
+
+		for (const FPropertyBagPropertyDesc& Desc : BagStruct->GetPropertyDescs())
+		{
+			// Try to extract as float (covers float and double bag types).
+			float FloatValue = 0.0f;
+			if (!TryGetFloat(Store, Desc, FloatValue))
+			{
+				// Non-float type (bool, int32, etc.) -forward to MetaSound only.
+				if (MSGen.IsValid())
+				{
+					ForwardToMetaSound(*MSGen, Store, Desc);
+				}
+				continue;
+			}
+
+			// If MetaSound has a matching input, send it there and skip the built-in handler.
+			if (MSGen.IsValid() && HasMetaSoundInput(Desc.Name))
+			{
+				MSGen->SetInputValue<float>(Desc.Name, FloatValue);
+				continue;
+			}
+
+			// Check if this is a built-in parameter with a dedicated DSP handler.
+			if (const FBuiltInParamHandler* Handler = BuiltInParams.Find(Desc.Name))
+			{
+				(this->**Handler)(FloatValue);
+				continue;
+			}
+
+			// Unknown float param with no MetaSound input -ignore.
+		}
+	}
+
+	void FSubsonicGenerator::SetVolume(float dB)
+	{
+		VolumeLinear = Audio::ConvertToLinear(dB);
+	}
+
+	void FSubsonicGenerator::SetPitchShift(float Semitones)
+	{
+		constexpr int32 PitchShiftInterpolationFrames = 64;
+		const float Ratio = FMath::Pow(2.0f, Semitones / 12.0f);
+		Resampler->SetFrameRatio(Ratio, PitchShiftInterpolationFrames);
+		if (!FMath::IsNearlyZero(Semitones))
+		{
+			bPitchShiftActive = true;
+		}
+	}
+
+	void FSubsonicGenerator::SetHighpassCutoff(float FreqHz)
+	{
+		if (FreqHz <= 0.0f)
+		{
+			HighpassFilter.SetEnabled(false);
+		}
+		else
+		{
+			HighpassFilter.SetEnabled(true);
+			HighpassFilter.SetFrequency(FreqHz);
+		}
+	}
+
+	void FSubsonicGenerator::SetLowpassCutoff(float FreqHz)
+	{
+		if (FreqHz <= 0.0f || FreqHz >= SampleRate * 0.5f)
+		{
+			LowpassFilter.SetEnabled(false);
+		}
+		else
+		{
+			LowpassFilter.SetEnabled(true);
+			LowpassFilter.SetFrequency(FreqHz);
+		}
+	}
+
+	void FSubsonicGenerator::SetFadeOutTime(float Seconds)
+	{
+		if (FadeGainDelta == 0.0f)
+		{
+			FadeOutTimeSec = Seconds;
+		}
+	}
+
+	void FSubsonicGenerator::SetSource(Audio::IAudioMixerGeneratorSource* InSource)
+	{
+		Source = InSource;
+	}
+
+	void FSubsonicGenerator::Play()
+	{
+		if (Source)
+		{
+			Source->Play();
+		}
+	}
+
+	void FSubsonicGenerator::Stop()
+	{
+		if (FadeGainDelta != 0.0f)
+		{
+			return;
+		}
+
+		if (FadeOutTimeSec > 0.0f)
+		{
+			// Begin fade -OnGenerateAudio will null Source when fade completes.
+			FadeGainDelta = 1.0f / (FadeOutTimeSec * SampleRate);
+			return;
+		}
+
+		// Immediate stop.
+		if (Source)
+		{
+			Source->Stop();
+			Source = nullptr;
+		}
+	}
+
+	Audio::IAudioMixerGeneratorSource* FSubsonicGenerator::GetSource() const
+	{
+		return Source;
+	}
+
+	int32 FSubsonicGenerator::OnGenerateAudio(float* OutAudio, int32 InNumSamples)
+	{
+		// Apply deferred params once graph becomes ready.
+		if (DeferredParams && bGraphReady.load(std::memory_order_acquire))
+		{
+			ApplyParameters(*DeferredParams);
+			DeferredParams.Reset();
+		}
+
+		const int32 NumChannels = InnerGenerator->GetNumChannels();
+		const int32 NumFrames = (NumChannels > 0) ? (InNumSamples / NumChannels) : InNumSamples;
+
+		int32 Result = 0;
+
+		// Signal chain: InnerGenerator ->[Resample/PitchShift] ->[Highpass] ->[Lowpass] ->Volume ->Output
+		if (bPitchShiftActive)
+		{
+			// Pitch shift active: generate more/fewer inner samples, then resample to output size.
+			const int32 NumInputFrames = Resampler->GetNumInputFramesNeededToProduceOutputFrames(NumFrames);
+			int32 NumInputSamples = NumInputFrames * NumChannels;
+
+			const int32 LeftoverSamples = ResampleInputBuffer.Num();
+			ResampleInputBuffer.SetNumUninitialized(FMath::Max(LeftoverSamples, NumInputSamples));
+			if (NumInputSamples > LeftoverSamples)
+			{
+				NumInputSamples = LeftoverSamples + InnerGenerator->OnGenerateAudio(ResampleInputBuffer.GetData() + LeftoverSamples, NumInputSamples - LeftoverSamples);
+			}
+			else
+			{
+				NumInputSamples = LeftoverSamples;
+			}
+
+			int32 InputFramesConsumed = 0;
+			int32 OutputFramesProduced = 0;
+			Resampler->ProcessInterleaved(
+				TArrayView<const float>(ResampleInputBuffer.GetData(), NumInputSamples),
+				TArrayView<float>(OutAudio, InNumSamples),
+				InputFramesConsumed,
+				OutputFramesProduced);
+
+			const int32 ProducedSamples = OutputFramesProduced * NumChannels;
+			Result = ProducedSamples;
+
+			// Save leftover frames for next run.
+			const int32 InputSamplesConsumed = InputFramesConsumed * NumChannels;
+			const int32 NewLeftoverSamples = NumInputSamples - InputSamplesConsumed;
+			if (NewLeftoverSamples > 0)
+			{
+				FMemory::Memmove(ResampleInputBuffer.GetData(), ResampleInputBuffer.GetData() + InputSamplesConsumed, NewLeftoverSamples * sizeof(float));
+			}
+			ResampleInputBuffer.SetNum(NewLeftoverSamples);
+		}
+		else
+		{
+			Result = InnerGenerator->OnGenerateAudio(OutAudio, InNumSamples);
+		}
+
+		// Apply filters in-place (no-op when disabled).
+		HighpassFilter.ProcessAudio(OutAudio, InNumSamples, OutAudio);
+		LowpassFilter.ProcessAudio(OutAudio, InNumSamples, OutAudio);
+
+		// Apply volume and fade-out in a single pass via SIMD-optimized ramp.
+		{
+			const float FadeStart = FadeGain;
+			if (FadeGainDelta != 0.0f)
+			{
+				FadeGain = FMath::Max(0.0f, FadeGain - FadeGainDelta * NumFrames);
+			}
+
+			const float StartGain = VolumeLinear * FadeStart;
+			const float EndGain = VolumeLinear * FadeGain;
+			if (StartGain != 1.0f || EndGain != 1.0f)
+			{
+				Audio::ArrayFade(TArrayView<float>(OutAudio, InNumSamples), StartGain, EndGain);
+			}
+		}
+
+		// When the generator is done (fade finished or inner generator ended),
+		// stop the source and null the pointer so the subscriber knows to clean up.
+		if (Source && IsFinished())
+		{
+			Source->Stop();
+			Source = nullptr;
+		}
+
+		return Result;
+	}
+
+	int32 FSubsonicGenerator::GetNumChannels() const
+	{
+		return InnerGenerator->GetNumChannels();
+	}
+
+	bool FSubsonicGenerator::IsFinished() const
+	{
+		return FadeGain <= 0.0f || InnerGenerator->IsFinished();
+	}
+
+} // namespace UE::Subsonic
diff --git a/Engine/Plugins/Experimental/Subsonic/Source/SubsonicEngine/Private/StandardEventSubscribers/SubsonicGenerator.h b/Engine/Plugins/Experimental/Subsonic/Source/SubsonicEngine/Private/StandardEventSubscribers/SubsonicGenerator.h
new file mode 100644
--- /dev/null
+++ b/Engine/Plugins/Experimental/Subsonic/Source/SubsonicEngine/Private/StandardEventSubscribers/SubsonicGenerator.h
@@ -0,0 +1,126 @@
+// Copyright Epic Games, Inc. All Rights Reserved.
+
+#pragma once
+
+#include "DSP/Filter.h"
+#include "DSP/RuntimeResampler.h"
+#include "MetasoundGenerator.h"
+#include "Sound/SoundGenerator.h"
+
+#include <atomic>
+
+namespace Audio { class IAudioMixerGeneratorSource; }
+
+namespace UE::Subsonic
+{
+	class FWaveGenerator;
+	struct FSubsonicParameterStore;
+
+	// Wrapper ISoundGenerator that owns either a WaveGenerator or MetsoundGenerator inner generator.
+	// Applies DSP (volume, pitch shift, highpass, lowpass) after delegating audio
+	//  generation to the inner generator.
+	//
+	// After initialization, these generators are owned by the audio render thread and should be communicated
+	//  with only through SubsonicRelay.
+	class FSubsonicGenerator : public ISoundGenerator
+	{
+	public:
+		// Wave path: inner generator is an FWaveGenerator (pure wave player).
+		explicit FSubsonicGenerator(TSharedRef<FWaveGenerator> InWave, float InSampleRate);
+
+		// MetaSound path: inner generator is an FMetasoundGenerator (returned by UMetaSoundSource).
+		explicit FSubsonicGenerator(ISoundGeneratorPtr InMetaSoundGenerator, float InSampleRate);
+
+		virtual ~FSubsonicGenerator() override;
+
+		//~ Begin ISoundGenerator
+		virtual int32 OnGenerateAudio(float* OutAudio, int32 InNumSamples) override;
+		virtual int32 GetNumChannels() const override;
+		virtual bool IsFinished() const override;
+		//~ End ISoundGenerator
+
+		// Audio render thread: apply named parameters from the store. Built-in parameters are
+		// handled by dedicated DSP handlers, unless the MetasoundGenerator provides a matching input.
+		void ApplyParameters(const FSubsonicParameterStore& Store);
+
+		// Associate a GeneratorSource with this generator, so the generator can be in charge of source state
+		//  (play, stop, pause, etc.)
+		void SetSource(Audio::IAudioMixerGeneratorSource* InSource);
+
+		// Start playback via the associated GeneratorSource.
+		void Play();
+
+		// Stop playback. If FadeOutTime > 0, begins a fade; otherwise calls Source->Stop() immediately.
+		void Stop();
+
+		// Game thread: returns the bare source pointer. Null means the generator is done
+		//  and the subscriber can safely destroy the source.
+		Audio::IAudioMixerGeneratorSource* GetSource() const;
+
+	private:
+		using FBuiltInParamHandler = void (FSubsonicGenerator::*)(float);
+		static const TMap<FName, FBuiltInParamHandler>& GetBuiltInParams();
+
+		// Attempts to extract a float from a property bag descriptor (handles float and double types).
+		static bool TryGetFloat(const FSubsonicParameterStore& Store, const FPropertyBagPropertyDesc& Desc, float& OutValue);
+
+		// Returns true if the pinned MetaSound generator has an input with the given name.
+		bool HasMetaSoundInput(FName Name) const;
+
+		// Forwards a non-float parameter to the pinned MetaSound generator.
+		void ForwardToMetaSound(Metasound::FMetasoundGenerator& MSGen, const FSubsonicParameterStore& Store, const FPropertyBagPropertyDesc& Desc);
+
+		// Built-in parameter handlers.
+		void SetVolume(float dB);
+		void SetPitchShift(float Semitones);
+		void SetHighpassCutoff(float FreqHz);
+		void SetLowpassCutoff(float FreqHz);
+		void SetFadeOutTime(float Seconds);
+
+		// Initializes DSP state (filters, resampler) after InnerGenerator is set.
+		void InitDSP();
+
+		ISoundGeneratorPtr InnerGenerator;
+
+		// Non-null for MetaSound path only.
+		TWeakPtr<Metasound::FMetasoundGenerator, ESPMode::ThreadSafe> MetaSoundGeneratorWeak;
+
+		// True once the MetaSound graph has been built and its inputs are available.
+		// Set via AddGraphSetCallback; always true for wave generators.
+		std::atomic<bool> bGraphReady{ false };
+
+		// Handle for the graph-set callback so we can unregister in the destructor.
+		FDelegateHandle GraphSetCallbackHandle;
+
+		// Holds the most recent SetParameters payload when the MetaSound graph is not yet ready.
+		// Applied and cleared once bGraphReady becomes true. Last-writer-wins.
+		TSharedPtr<FSubsonicParameterStore> DeferredParams;
+
+		// Current volume as linear gain. Updated from the "Volume" parameter (interpreted as dB).
+		// 0 dB = 1.0 linear (unity gain). Render thread only.
+		float VolumeLinear = 1.0f;
+
+		// Cached sample rate from the audio device.
+		float SampleRate = 0.0f;
+
+		// DSP processors applied in the signal chain after the inner generator.
+		Audio::FBiquadFilter HighpassFilter;
+		Audio::FBiquadFilter LowpassFilter;
+
+		// Pitch shifting via resampling.
+		TUniquePtr<Audio::FRuntimeResampler> Resampler;
+		TArray<float> ResampleInputBuffer;
+		bool bPitchShiftActive = false;
+
+		// Fade-out envelope state. Render thread only.
+		// FadeGainDelta > 0 means a fade is in progress; FadeGain <= 0 means it has completed.
+		float FadeOutTimeSec = 0.0f;
+		float FadeGain = 1.0f;
+		float FadeGainDelta = 0.0f;
+
+		// Non-owning pointer to the associated GeneratorSource. Set on the game thread
+		//  at creation, nulled on the render thread when the generator is done.
+		Audio::IAudioMixerGeneratorSource* Source = nullptr;
+	};
+
+} // namespace UE::Subsonic
diff --git a/Engine/Plugins/Experimental/Subsonic/Source/SubsonicEngine/Private/StandardEventSubscribers/SubsonicGeneratorSourceSubscriber.cpp b/Engine/Plugins/Experimental/Subsonic/Source/SubsonicEngine/Private/StandardEventSubscribers/SubsonicGeneratorSourceSubscriber.cpp
new file mode 100644
--- /dev/null
+++ b/Engine/Plugins/Experimental/Subsonic/Source/SubsonicEngine/Private/StandardEventSubscribers/SubsonicGeneratorSourceSubscriber.cpp
@@ -0,0 +1,376 @@
+// Copyright Epic Games, Inc. All Rights Reserved.
+
+#include "StandardEventSubscribers/SubsonicGeneratorSourceSubscriber.h"
+
+#include "Algo/Transform.h"
+#include "AudioDevice.h"
+#include "AudioDeviceManager.h"
+#include "AudioMixerDevice.h"
+#include "Engine/Engine.h"
+#include "IAudioMixerGeneratorSource.h"
+#include "MetasoundSource.h"
+#include "StandardEventSubscribers/SubsonicRelay.h"
+#include "SubsonicCoreLog.h"
+#include "SubsonicGenerator.h"
+#include "SubsonicSubsystem.h"
+#include "SubsonicWaveGenerator.h"
+#include "UObject/UObjectGlobals.h"
+
+
+namespace UE::Subsonic
+{
+	void FInstanceSources::ApplyParameters(FSubsonicRelay& Relay)
+	{
+		if (!Generator)
+		{
+			return;
+		}
+
+		// Always allocate a fresh merged store - the audio thread owns its TSharedPtr.
+		TSharedPtr<FSubsonicParameterStore> Merged = MakeShared<FSubsonicParameterStore>();
+		Merged->MergeFrom(AuthoredParams);
+		Merged->MergeFrom(RuntimeParams);
+
+		// Enqueue for the audio thread via the relay's global command queue.
+		FRelayCommand Cmd;
+		Cmd.Type = FRelayCommand::EType::SetParameters;
+		Cmd.Params = MoveTemp(Merged);
+		Cmd.Target = Generator;
+		Relay.EnqueueCommand(MoveTemp(Cmd));
+	}
+
+	namespace SubscriberPrivate
+	{
+		// Creates an FSubsonicGenerator from a USoundWave, using the MetaSound path
+		// if the sound is a UMetaSoundSource, otherwise the wave path.
+		TSharedPtr<FSubsonicGenerator> CreateGenerator(USoundWave& Sound, Audio::FMixerDevice& AudioDevice)
+		{
+			const float SampleRate = AudioDevice.GetSampleRate();
+
+			if (UMetaSoundSource* MetaSound = Cast<UMetaSoundSource>(&Sound))
+			{
+				// Ensure the MetaSound graph is registered and compiled.
+				MetaSound->InitResources();
+
+				FSoundGeneratorInitParams InitParams;
+				InitParams.AudioDeviceID = AudioDevice.DeviceID;
+				InitParams.SampleRate = SampleRate;
+				InitParams.AudioMixerNumOutputFrames = AudioDevice.GetNumOutputFrames();
+				InitParams.NumChannels = MetaSound->NumChannels;
+				InitParams.NumFramesPerCallback = AudioDevice.GetNumOutputFrames();
+				InitParams.AudioComponentId = 0;
+				InitParams.InstanceID = 0;
+				InitParams.bIsPreviewSound = false;
+
+				TArray<FAudioParameter> DefaultParams;
+				MetaSound->InitParameters(DefaultParams, NAME_None);
+				ISoundGeneratorPtr MSGenPtr = MetaSound->CreateSoundGenerator(InitParams, MoveTemp(DefaultParams));
+				if (!MSGenPtr.IsValid())
+				{
+					return nullptr;
+				}
+
+				return MakeShared<FSubsonicGenerator>(MSGenPtr, SampleRate);
+			}
+
+			TSharedRef<FWaveGenerator> WaveGen = MakeShared<FWaveGenerator>(Sound.CreateSoundWaveProxy(), &AudioDevice);
+			return MakeShared<FSubsonicGenerator>(WaveGen, SampleRate);
+		}
+
+		template <typename HandleType>
+		Audio::IAudioMixerGeneratorSource* PlaySound(
+			USubsonicGeneratorSourceSubscriber& Subscriber,
+			USoundWave& Sound,
+			const HandleType& InHandle,
+			TMap<HandleType, FGeneratorSourceCache>& OutSourceCache,
+			FName Name,
+			const FSubsonicParameterStore& Params,
+			FSubsonicRelay& Relay,
+			bool bReleaseExisting)
+		{
+			Audio::IAudioMixerGeneratorSource* ReturnSource = nullptr;
+
+			if (FAudioDeviceManager* AudioDeviceManager = FAudioDeviceManager::Get())
+			{
+				if (Audio::FMixerDevice* AudioDevice = Subscriber.GetMutableMixerDevice())
+				{
+					TSharedPtr<FSubsonicGenerator> Generator = CreateGenerator(Sound, *AudioDevice);
+					if (!Generator.IsValid())
+					{
+						return nullptr;
+					}
+
+					FGeneratorSourceCache& Cache = OutSourceCache.FindOrAdd(InHandle);
+					FInstanceSources& InstanceSources = Cache.FindOrAdd(Name);
+
+					TUniquePtr<Audio::IAudioMixerGeneratorSource> NewSource = AudioDevice->CreateGeneratorSource(Generator.ToSharedRef());
+					if (NewSource.IsValid())
+					{
+						ReturnSource = NewSource.Get();
+						if (InstanceSources.ActiveSource != nullptr)
+						{
+							FReleasedSource Released;
+							Released.Source = MoveTemp(InstanceSources.ActiveSource);
+							Released.Generator = InstanceSources.Generator;
+							InstanceSources.ReleasedSources.Add(MoveTemp(Released));
+						}
+
+						InstanceSources.Generator = Generator;
+						InstanceSources.AuthoredParams = Params;
+						InstanceSources.RuntimeParams.Reset();
+						InstanceSources.ApplyParameters(Relay);
+
+						Generator->SetSource(NewSource.Get());
+						Relay.EnqueueCommand(FRelayCommand{ FRelayCommand::EType::Play, nullptr, Generator });
+						InstanceSources.ActiveSource = MoveTemp(NewSource);
+					}
+				}
+			}
+
+			return ReturnSource;
+		}
+
+		template <typename HandleType>
+		void ClearStaleSounds(TMap<HandleType, FGeneratorSourceCache>& SourceCache)
+		{
+			TArray<HandleType> HandlesToClear;
+
+			for (TPair<HandleType, FGeneratorSourceCache>& CacheItem : SourceCache)
+			{
+				TArray<FName> NamesToClear;
+				for (TPair<FName, FInstanceSources>& SourcesItem : CacheItem.Value)
+				{
+					FInstanceSources& Instances = SourcesItem.Value;
+
+					Instances.ReleasedSources.RemoveAll([](const FReleasedSource& Released)
+						{
+							return !Released.Generator || Released.Generator->GetSource() == nullptr;
+						});
+					if (Instances.ActiveSource != nullptr && Instances.Generator && Instances.Generator->GetSource() == nullptr)
+					{
+						Instances.ActiveSource.Reset();
+						Instances.Generator.Reset();
+					}
+					if (Instances.ActiveSource == nullptr && Instances.ReleasedSources.IsEmpty())
+					{
+						NamesToClear.Add(SourcesItem.Key);
+					}
+				}
+
+				for (FName Name : NamesToClear)
+				{
+					CacheItem.Value.Remove(Name);
+				}
+				if (CacheItem.Value.IsEmpty())
+				{
+					HandlesToClear.Add(CacheItem.Key);
+				}
+			}
+
+			for (HandleType Handle : HandlesToClear)
+			{
+				SourceCache.Remove(Handle);
+			}
+		}
+
+		template <typename HandleType>
+		void StopSound(
+			FSubsonicRelay& Relay,
+			const HandleType& InHandle,
+			TMap<HandleType, FGeneratorSourceCache>& OutSourceCache,
+			FName Name)
+		{
+			FGeneratorSourceCache* Cache = OutSourceCache.Find(InHandle);
+			if (Cache != nullptr)
+			{
+				if (FInstanceSources* Instances = Cache->Find(Name))
+				{
+					if (Instances->ActiveSource != nullptr)
+					{
+						FReleasedSource Released;
+						Released.Source = MoveTemp(Instances->ActiveSource);
+						Released.Generator = Instances->Generator;
+						Instances->ReleasedSources.Add(MoveTemp(Released));
+
+						Relay.EnqueueCommand(FRelayCommand{
+							FRelayCommand::EType::Stop, nullptr, Instances->Generator
+						});
+
+						Instances->Generator.Reset();
+					}
+				}
+			}
+		}
+
+		template <typename HandleType>
+		void SetParameters(
+			FSubsonicRelay& Relay,
+			const HandleType& InHandle,
+			TMap<HandleType, FGeneratorSourceCache>& SourceCache,
+			FName SourceName,
+			const FSubsonicParameterStore& Params)
+		{
+			if (FGeneratorSourceCache* Cache = SourceCache.Find(InHandle))
+			{
+				if (FInstanceSources* Instances = Cache->Find(SourceName))
+				{
+					Instances->RuntimeParams.MergeFrom(Params);
+					Instances->ApplyParameters(Relay);
+				}
+			}
+		}
+
+		template <typename HandleType>
+		void SetAuthoredParameters(
+			FSubsonicRelay& Relay,
+			const HandleType& InHandle,
+			TMap<HandleType, FGeneratorSourceCache>& SourceCache,
+			FName SourceName,
+			const FSubsonicParameterStore& Params)
+		{
+			if (FGeneratorSourceCache* Cache = SourceCache.Find(InHandle))
+			{
+				if (FInstanceSources* Instances = Cache->Find(SourceName))
+				{
+					Instances->AuthoredParams = Params;
+					Instances->ApplyParameters(Relay);
+				}
+			}
+		}
+	} // namespace SubscriberPrivate
+
+	USubsonicGeneratorSourceSubscriber::USubsonicGeneratorSourceSubscriber()
+		: FTickableGameObject(ETickableTickType::Never)
+		, Relay(MakeShared<FSubsonicRelay>())
+	{
+	}
+
+	void USubsonicGeneratorSourceSubscriber::Initialize(FSubsystemCollectionBase& Collection)
+	{
+		Super::Initialize(Collection);
+		SetTickableTickType(ETickableTickType::Always);
+	}
+
+	void USubsonicGeneratorSourceSubscriber::Deinitialize()
+	{
+		SetTickableTickType(ETickableTickType::Never);
+
+		GlobalSources.Reset();
+		CollectionSources.Reset();
+		EventSources.Reset();
+		ActionSources.Reset();
+
+		Super::Deinitialize();
+	}
+
+	void USubsonicGeneratorSourceSubscriber::OnCollectionRegistered(const FCollectionHandle& InHandle)
+	{
+	}
+
+	void USubsonicGeneratorSourceSubscriber::OnCollectionUnregistered(const FCollectionHandle& InHandle)
+	{
+		CollectionSources.Remove(InHandle);
+	}
+
+	void USubsonicGeneratorSourceSubscriber::PlaySound(FName Name, USoundWave& Sound, const FSubsonicParameterStore& Params, bool bReleaseExisting)
+	{
+		const int32 FakeHandle = 0;
+		SubscriberPrivate::PlaySound(*this, Sound, FakeHandle, GlobalSources, Name, Params, *Relay, bReleaseExisting);
+	}
+
+	void USubsonicGeneratorSourceSubscriber::PlaySound(const FCollectionHandle& InHandle, FName Name, USoundWave& Sound, const FSubsonicParameterStore& Params, bool bReleaseExisting)
+	{
+		SubscriberPrivate::PlaySound(*this, Sound, InHandle, CollectionSources, Name, Params, *Relay, bReleaseExisting);
+	}
+
+	void USubsonicGeneratorSourceSubscriber::PlaySound(const FEventHandle& InHandle, FName Name, USoundWave& Sound, const FSubsonicParameterStore& Params, bool bReleaseExisting)
+	{
+		SubscriberPrivate::PlaySound(*this, Sound, InHandle, EventSources, Name, Params, *Relay, bReleaseExisting);
+	}
+
+	void USubsonicGeneratorSourceSubscriber::PlaySound(const FActionHandle& InHandle, FName Name, USoundWave& Sound, const FSubsonicParameterStore& Params, bool bReleaseExisting)
+	{
+		SubscriberPrivate::PlaySound(*this, Sound, InHandle, ActionSources, Name, Params, *Relay, bReleaseExisting);
+	}
+
+	void USubsonicGeneratorSourceSubscriber::StopSound(FName Name)
+	{
+		const int32 FakeHandle = 0;
+		SubscriberPrivate::StopSound(*Relay, FakeHandle, GlobalSources, Name);
+	}
+
+	void USubsonicGeneratorSourceSubscriber::StopSound(const FCollectionHandle& InHandle, FName Name)
+	{
+		SubscriberPrivate::StopSound(*Relay, InHandle, CollectionSources, Name);
+	}
+
+	void USubsonicGeneratorSourceSubscriber::StopSound(const FEventHandle& InHandle, FName Name)
+	{
+		SubscriberPrivate::StopSound(*Relay, InHandle, EventSources, Name);
+	}
+
+	void USubsonicGeneratorSourceSubscriber::StopSound(const FActionHandle& InHandle, FName Name)
+	{
+		SubscriberPrivate::StopSound(*Relay, InHandle, ActionSources, Name);
+	}
+
+	void USubsonicGeneratorSourceSubscriber::SetParameters(FName SourceName, const FSubsonicParameterStore& Params)
+	{
+		const int32 FakeHandle = 0;
+		SubscriberPrivate::SetParameters(*Relay, FakeHandle, GlobalSources, SourceName, Params);
+	}
+
+	void USubsonicGeneratorSourceSubscriber::SetParameters(const FCollectionHandle& InHandle, FName SourceName, const FSubsonicParameterStore& Params)
+	{
+		SubscriberPrivate::SetParameters(*Relay, InHandle, CollectionSources, SourceName, Params);
+	}
+
+	void USubsonicGeneratorSourceSubscriber::SetParameters(const FEventHandle& InHandle, FName SourceName, const FSubsonicParameterStore& Params)
+	{
+		SubscriberPrivate::SetParameters(*Relay, InHandle, EventSources, SourceName, Params);
+	}
+
+	void USubsonicGeneratorSourceSubscriber::SetParameters(const FActionHandle& InHandle, FName SourceName, const FSubsonicParameterStore& Params)
+	{
+		SubscriberPrivate::SetParameters(*Relay, InHandle, ActionSources, SourceName, Params);
+	}
+
+	void USubsonicGeneratorSourceSubscriber::SetAuthoredParameters(FName SourceName, const FSubsonicParameterStore& Params)
+	{
+		const int32 FakeHandle = 0;
+		SubscriberPrivate::SetAuthoredParameters(*Relay, FakeHandle, GlobalSources, SourceName, Params);
+	}
+
+	void USubsonicGeneratorSourceSubscriber::SetAuthoredParameters(const FCollectionHandle& InHandle, FName SourceName, const FSubsonicParameterStore& Params)
+	{
+		SubscriberPrivate::SetAuthoredParameters(*Relay, InHandle, CollectionSources, SourceName, Params);
+	}
+
+	void USubsonicGeneratorSourceSubscriber::SetAuthoredParameters(const FEventHandle& InHandle, FName SourceName, const FSubsonicParameterStore& Params)
+	{
+		SubscriberPrivate::SetAuthoredParameters(*Relay, InHandle, EventSources, SourceName, Params);
+	}
+
+	void USubsonicGeneratorSourceSubscriber::SetAuthoredParameters(const FActionHandle& InHandle, FName SourceName, const FSubsonicParameterStore& Params)
+	{
+		SubscriberPrivate::SetAuthoredParameters(*Relay, InHandle, ActionSources, SourceName, Params);
+	}
+
+	void USubsonicGeneratorSourceSubscriber::Tick(float DeltaTime)
+	{
+		if (Audio::FMixerDevice* MixerDevice = GetMutableMixerDevice())
+		{
+			Relay->Tick(*MixerDevice);
+		}
+
+		SubscriberPrivate::ClearStaleSounds(GlobalSources);
+		SubscriberPrivate::ClearStaleSounds(CollectionSources);
+		SubscriberPrivate::ClearStaleSounds(EventSources);
+		SubscriberPrivate::ClearStaleSounds(ActionSources);
+	}
+
+	TStatId USubsonicGeneratorSourceSubscriber::GetStatId() const
+	{
+		RETURN_QUICK_DECLARE_CYCLE_STAT(USubsonicGeneratorSourceSubscriber, STATGROUP_Tickables);
+	}
+}
diff --git a/Engine/Plugins/Experimental/Subsonic/Source/SubsonicEngine/Private/StandardEventSubscribers/SubsonicRelay.cpp b/Engine/Plugins/Experimental/Subsonic/Source/SubsonicEngine/Private/StandardEventSubscribers/SubsonicRelay.cpp
new file mode 100644
--- /dev/null
+++ b/Engine/Plugins/Experimental/Subsonic/Source/SubsonicEngine/Private/StandardEventSubscribers/SubsonicRelay.cpp
@@ -0,0 +1,57 @@
+// Copyright Epic Games, Inc. All Rights Reserved.
+
+#include "StandardEventSubscribers/SubsonicRelay.h"
+
+#include "AudioMixerDevice.h"
+#include "SubsonicGenerator.h"
+
+namespace UE::Subsonic
+{
+	void FSubsonicRelay::EnqueueCommand(FRelayCommand&& Cmd)
+	{
+		FScopeLock ScopeLock(&CommandLock);
+		PendingCommands.Add(MoveTemp(Cmd));
+	}
+
+	void FSubsonicRelay::Tick(Audio::FMixerDevice& MixerDevice)
+	{
+		TArray<FRelayCommand> Commands;
+		{
+			FScopeLock ScopeLock(&CommandLock);
+			Commands = MoveTemp(PendingCommands);
+			// PendingCommands is now empty after move
+		}
+
+		if (Commands.IsEmpty())
+		{
+			return;
+		}
+
+		MixerDevice.AudioRenderThreadCommand([Cmds = MoveTemp(Commands)]()
+		{
+			for (const FRelayCommand& Cmd : Cmds)
+			{
+				if (!Cmd.Target)
+				{
+					continue;
+				}
+
+				switch (Cmd.Type)
+				{
+				case FRelayCommand::EType::SetParameters:
+					if (Cmd.Params)
+					{
+						Cmd.Target->ApplyParameters(*Cmd.Params);
+					}
+					break;
+				case FRelayCommand::EType::Play:
+					Cmd.Target->Play();
+					break;
+				case FRelayCommand::EType::Stop:
+					Cmd.Target->Stop();
+					break;
+				}
+			}
+		});
+	}
+} // namespace UE::Subsonic
diff --git a/Engine/Plugins/Experimental/Subsonic/Source/SubsonicEngine/Private/StandardEventSubscribers/SubsonicWaveGenerator.cpp b/Engine/Plugins/Experimental/Subsonic/Source/SubsonicEngine/Private/StandardEventSubscribers/SubsonicWaveGenerator.cpp
new file mode 100644
--- /dev/null
+++ b/Engine/Plugins/Experimental/Subsonic/Source/SubsonicEngine/Private/StandardEventSubscribers/SubsonicWaveGenerator.cpp
@@ -0,0 +1,74 @@
+// Copyright Epic Games, Inc. All Rights Reserved.
+
+#include "SubsonicWaveGenerator.h"
+
+#include "AudioMixerDevice.h"
+#include "DSP/MultichannelBuffer.h"
+#include "Sound/SoundGenerator.h"
+#include "Sound/SoundWave.h"
+#include "Sound/SoundWaveProxyPlayer.h"
+
+namespace UE::Subsonic
+{
+	namespace WaveGeneratorPrivate
+	{
+		void Interleave(const Audio::FMultichannelBufferView& Input, float* Output, const int32 NumOutputSamples)
+		{
+			check(Input.IsEmpty() || NumOutputSamples == Input.Num() * Input[0].Num());
+			for (int32 ChannelIndex = 0; ChannelIndex < Input.Num(); ++ChannelIndex)
+			{
+				float* OutputSample = Output + ChannelIndex;
+				TArrayView<float> InputChannel = Input[ChannelIndex];
+				for (float Sample : InputChannel)
+				{
+					*OutputSample = Sample;
+					OutputSample += Input.Num();
+				}
+			}
+		}
+	}
+
+	FWaveGenerator::FWaveGenerator(const TSharedPtr<FSoundWaveProxy>& SoundWaveProxy, Audio::FMixerDevice* MixerDevice)
+	{
+		const float SampleRate = MixerDevice->GetSampleRate();
+		check(SoundWaveProxy.IsValid());
+		FSoundWaveProxyPlayer::FSettings Settings(SampleRate, SoundWaveProxy->GetNumChannels());
+		Player = FSoundWaveProxyPlayer::Create(Settings);
+		Player->SetSoundWave(SoundWaveProxy);
+	}
+
+	FWaveGenerator::~FWaveGenerator()
+	{
+	}
+
+	int32 FWaveGenerator::OnGenerateAudio(float* OutAudio, int32 InNumSamples)
+	{
+		const int32 NumChannels = GetNumChannels();
+		check(InNumSamples >= 0 && NumChannels > 0 && InNumSamples % NumChannels == 0);
+		const int32 NumFrames = InNumSamples / NumChannels;
+
+		Audio::FMultichannelBuffer TempBuffer;
+		for (int32 i = 0; i < NumChannels; ++i)
+		{
+			Audio::FAlignedFloatBuffer ChannelBuffer;
+			ChannelBuffer.SetNumUninitialized(NumFrames);
+			TempBuffer.Add(MoveTemp(ChannelBuffer));
+		}
+
+		Audio::FMultichannelBufferView TempBufferView = Audio::MakeMultichannelBufferView(TempBuffer);
+		Player->RenderMultiChannelAudio(TempBufferView);
+		WaveGeneratorPrivate::Interleave(TempBufferView, OutAudio, InNumSamples);
+
+		return InNumSamples;
+	}
+
+	int32 FWaveGenerator::GetNumChannels() const
+	{
+		return Player->GetSourceNumChannels();
+	}
+
+	bool FWaveGenerator::IsFinished() const
+	{
+		return Player->IsFinished();
+	}
+} // namespace UE::Subsonic
\ No newline at end of file
diff --git a/Engine/Plugins/Experimental/Subsonic/Source/SubsonicEngine/Private/StandardEventSubscribers/SubsonicWaveGenerator.h b/Engine/Plugins/Experimental/Subsonic/Source/SubsonicEngine/Private/StandardEventSubscribers/SubsonicWaveGenerator.h
new file mode 100644
--- /dev/null
+++ b/Engine/Plugins/Experimental/Subsonic/Source/SubsonicEngine/Private/StandardEventSubscribers/SubsonicWaveGenerator.h
@@ -0,0 +1,34 @@
+// Copyright Epic Games, Inc. All Rights Reserved.
+
+#pragma once
+
+#include "Sound/SoundGenerator.h"
+
+namespace Audio
+{
+	class FMixerDevice;
+}
+
+class FSoundWaveProxy;
+class FSoundWaveProxyPlayer;
+
+namespace UE::Subsonic
+{
+
+	// A pure wave-file player. Volume scaling and relay communication are handled
+	// by the outer FSubsonicGenerator wrapper.
+	class FWaveGenerator : public ISoundGenerator
+	{
+	public:
+		FWaveGenerator(const TSharedPtr<FSoundWaveProxy>& SoundWaveProxy, Audio::FMixerDevice* MixerDevice);
+
+		virtual ~FWaveGenerator() override;
+
+		virtual int32 OnGenerateAudio(float* OutAudio, int32 InNumSamples) override;
+		virtual int32 GetNumChannels() const override;
+		virtual bool IsFinished() const override;
+
+	protected:
+		TUniquePtr<FSoundWaveProxyPlayer> Player;
+	};
+}
diff --git a/Engine/Plugins/Experimental/Subsonic/Source/SubsonicEngine/Private/SubsonicEngineModule.cpp b/Engine/Plugins/Experimental/Subsonic/Source/SubsonicEngine/Private/SubsonicEngineModule.cpp
new file mode 100644
--- /dev/null
+++ b/Engine/Plugins/Experimental/Subsonic/Source/SubsonicEngine/Private/SubsonicEngineModule.cpp
@@ -0,0 +1,20 @@
+// Copyright Epic Games, Inc. All Rights Reserved.
+
+#include "Modules/ModuleManager.h"
+
+namespace UE::Subsonic
+{
+	class FSubsonicEngineModule : public IModuleInterface
+	{
+	public:
+		virtual void StartupModule() override
+		{
+		}
+
+		virtual void ShutdownModule() override
+		{
+		}
+	};
+} // namespace UE::Subsonic
+
+IMPLEMENT_MODULE(UE::Subsonic::FSubsonicEngineModule, SubsonicEngine)
diff --git a/Engine/Plugins/Experimental/Subsonic/Source/SubsonicEngine/Private/SubsonicEventCollectionObjects.cpp b/Engine/Plugins/Experimental/Subsonic/Source/SubsonicEngine/Private/SubsonicEventCollectionObjects.cpp
new file mode 100644
--- /dev/null
+++ b/Engine/Plugins/Experimental/Subsonic/Source/SubsonicEngine/Private/SubsonicEventCollectionObjects.cpp
@@ -0,0 +1,204 @@
+// Copyright Epic Games, Inc. All Rights Reserved.
+
+#include "SubsonicEventCollectionObjects.h"
+
+#include "AudioDevice.h"
+#include "Engine/World.h"
+#include "SubsonicExecutor.h"
+#include "SubsonicHandles.h"
+
+
+namespace UE::Subsonic
+{
+	const Core::FSubsonicEventCollectionDefinition& USubsonicEventCollection::GetDefinition() const
+	{
+		return CollectionDefinition;
+	}
+
+#if WITH_EDITOR
+	FName USubsonicEventCollection::GetDefinitionPropertyName()
+	{
+		return GET_MEMBER_NAME_CHECKED(USubsonicEventCollection, CollectionDefinition);
+	}
+#endif // WITH_EDITOR
+
+	Core::FCollectionHandle USubsonicEventCollection::GetHandle() const
+	{
+#if WITH_EDITORONLY_DATA
+		FName CollectionName;
+		if (const UPackage* CollectionPackage = GetPackage())
+		{
+			const FTopLevelAssetPath AssetPath(CollectionPackage);
+			CollectionName = AssetPath.GetPackageName();
+		}
+#endif // WITH_EDITORONLY_DATA
+
+		return Core::FCollectionHandle
+		{
+#if WITH_EDITORONLY_DATA
+			.CollectionName = CollectionName,
+#endif // WITH_EDITORONLY_DATA
+			.CollectionId = CollectionDefinition.GetCollectionId(),
+		};
+	}
+
+#if WITH_EDITOR
+	Core::FSubsonicEventCollectionDefinition& USubsonicEventCollection::GetMutableDefinition()
+	{
+		return CollectionDefinition;
+	}
+
+	void USubsonicEventCollection::PostEditChangeChainProperty(FPropertyChangedChainEvent& PropertyChangedEvent)
+	{
+		Super::PostEditChangeChainProperty(PropertyChangedEvent);
+
+#if WITH_EDITORONLY_DATA
+		for (auto* Node = PropertyChangedEvent.PropertyChain.GetHead(); Node; Node = Node->GetNextNode())
+		{
+			const FName PropName = Node->GetValue()->GetFName();
+			if (PropName == Core::FSubsonicEventCollectionDefinition::GetParametersPropertyName()
+				|| PropName == Core::FSubsonicEvent::GetParametersPropertyName()
+				|| PropName == Core::FSubsonicEventActionBase::GetBindingsParameterName())
+			{
+				const bool bBindingsRemoved = CollectionDefinition.RemoveStaleBindings(GetHandle(), true /* bPromptRemoval */);
+				if (bBindingsRemoved)
+				{
+					CollectionDefinition.BindAllParameters();
+				}
+				break;
+			}
+		}
+#endif // WITH_EDITORONLY_DATA
+	}
+
+	void USubsonicEventCollection::PostEditUndo()
+	{
+		Super::PostEditUndo();
+
+		CollectionDefinition.RemoveStaleBindings(GetHandle());
+		CollectionDefinition.BindAllParameters();
+	}
+#endif // WITH_EDITOR
+
+
+	void USubsonicEventCollection::SetDefinition(Core::FSubsonicEventCollectionDefinition&& Definition)
+	{
+		CollectionDefinition = MoveTemp(Definition);
+	}
+
+	void USubsonicEventCollection::PostInitProperties()
+	{
+		Super::PostInitProperties();
+		CollectionDefinition.AssignId();
+	}
+
+	void USubsonicEventCollection::PostLoad()
+	{
+		checkf(CollectionDefinition.IsValid(), TEXT("Definition should have an assigned Id by this point (via PostInitProperties)"));
+
+		Super::PostLoad();
+
+		Audio::FDeviceId DeviceId = INDEX_NONE;
+		if (const UWorld* World = GetWorld())
+		{
+			if (const FAudioDevice* AudioDevice = World->GetAudioDeviceRaw())
+			{
+				DeviceId = AudioDevice->DeviceID;
+			}
+		}
+
+		CollectionDefinition.Register(GetHandle(), DeviceId);
+	}
+
+	void USubsonicEventCollection::BeginDestroy()
+	{
+		CollectionDefinition.Unregister();
+		Super::BeginDestroy();
+	}
+
+	void USubsonicEventCollectionExecutor::BeginDestroy()
+	{
+		Executor.Unregister();
+		Super::BeginDestroy();
+	}
+
+	USubsonicEventCollectionExecutor* USubsonicEventCollectionExecutor::Create(UObject& Parent, FName Name, const USubsonicEventCollection& InCollection, Audio::FDeviceId DeviceId)
+	{
+		if (USubsonicEventCollectionExecutor* NewExecutor = NewObject<USubsonicEventCollectionExecutor>(&Parent, Name, RF_Transient))
+		{
+			class FEventCollectionAccessor : public Core::FSubsonicExecutor::ICollectionAccessor
+			{
+				// Parent USubsonicEventExecutor contains strong ptr, so no need
+				// to induce GC overhead by making this a strong pointer.
+				TWeakObjectPtr<const USubsonicEventCollection> Collection;
+
+			public:
+				FEventCollectionAccessor(const USubsonicEventCollection& InCol)
+					: Collection(TWeakObjectPtr<const USubsonicEventCollection>(&InCol))
+				{
+				}
+
+				virtual ~FEventCollectionAccessor() = default;
+
+				const Core::FSubsonicEventCollectionDefinition* GetDefinition() const override final
+				{
+					if (Collection.IsValid())
+					{
+						return &Collection->GetDefinition();
+					}
+
+					return nullptr;
+				}
+
+				Core::FCollectionHandle GetHandle() const override final
+				{
+					if (Collection.IsValid())
+					{
+						return Collection->GetHandle();
+					}
+
+					return Core::FCollectionHandle();
+				}
+			};
+
+			NewExecutor->Collection = &InCollection;
+			NewExecutor->Executor = Core::FSubsonicExecutor::Create(DeviceId, MakeUnique<FEventCollectionAccessor>(InCollection));
+			return NewExecutor;
+		}
+
+		return nullptr;
+	}
+
+	void USubsonicEventCollectionExecutor::ExecuteEvent(FGameplayTag EventTag, ESubsonicExecutionResult& OutResult)
+	{
+		const bool bExecuted = Executor.ExecuteEvent(EventTag.GetTagName());
+		OutResult = bExecuted ? ESubsonicExecutionResult::Succeeded : ESubsonicExecutionResult::Failed;
+	}
+
+	const USubsonicEventCollection& USubsonicEventCollectionExecutor::GetCollectionChecked() const
+	{
+		check(Collection.Get());
+		return *Collection.Get();
+	}
+
+	const USubsonicEventCollection* USubsonicEventCollectionExecutor::GetCollection() const
+	{
+		return Collection.Get();
+	}
+
+	const Core::FSubsonicExecutor& USubsonicEventCollectionExecutor::GetExecutor() const
+	{
+		return Executor;
+	}
+
+	bool USubsonicEventCollectionExecutor::IsValid() const
+	{
+		return Executor.IsValid() && GetCollection() != nullptr;
+	}
+
+	void USubsonicEventCollectionExecutor::Unregister()
+	{
+		Executor.Unregister();
+		Collection = nullptr;
+	}
+} // namespace UE::Subsonic
diff --git a/Engine/Plugins/Experimental/Subsonic/Source/SubsonicEngine/Private/SubsonicSubsystem.cpp b/Engine/Plugins/Experimental/Subsonic/Source/SubsonicEngine/Private/SubsonicSubsystem.cpp
new file mode 100644
--- /dev/null
+++ b/Engine/Plugins/Experimental/Subsonic/Source/SubsonicEngine/Private/SubsonicSubsystem.cpp
@@ -0,0 +1,200 @@
+// Copyright Epic Games, Inc. All Rights Reserved.
+
+#include "SubsonicSubsystem.h"
+
+#include "Algo/Transform.h"
+#include "AudioDefines.h"
+#include "Engine/Engine.h"
+#include "ISubsonicEventRegistry.h"
+#include "Misc/Optional.h"
+#include "SubsonicEventCollectionObjects.h"
+#include "SubsonicEventSubscriberInterface.h"
+#include "Subsystems/AudioEngineSubsystem.h"
+#include "UObject/GCObject.h"
+
+
+namespace UE::Subsonic
+{
+	namespace SubsystemPrivate
+	{
+		bool TryGetDeviceIdFromWorldObject(const UObject* WorldContextObject, Audio::FDeviceId& OutDeviceId)
+		{
+			OutDeviceId = INDEX_NONE;
+			if (WorldContextObject)
+			{
+				if (GEngine)
+				{
+					if (UWorld* World = GEngine->GetWorldFromContextObject(WorldContextObject, EGetWorldErrorMode::LogAndReturnNull))
+					{
+						OutDeviceId = World->GetAudioDevice().GetDeviceID();
+						return true;
+					}
+				}
+			}
+			else
+			{
+				if (FAudioDeviceManager* DeviceManager = FAudioDeviceManager::Get())
+				{
+					OutDeviceId = DeviceManager->GetActiveAudioDevice().GetDeviceID();
+					return true;
+				}
+			}
+
+			// Either engine isn't loaded, or context object has no parent world
+			return false;
+		};
+
+		class FSubsonicEventRegistry : public FGCObject, public Core::ISubsonicEventRegistry
+		{
+		public:
+			FSubsonicEventRegistry() = default;
+
+			virtual ~FSubsonicEventRegistry() = default;
+
+		private:
+			void IterateSubscribersMatchingDevice(Audio::FDeviceId DeviceId, TFunctionRef<void(TScriptInterface<Core::ISubsonicEventSubscriberInterface>&)> Func)
+			{
+				auto GetSubscriberDeviceId = [](TScriptInterface<Core::ISubsonicEventSubscriberInterface>& Subscriber) -> TOptional<Audio::FDeviceId>
+				{
+					UObject* SubscriberObject = Subscriber.GetObject();
+					if (!SubscriberObject)
+					{
+						return { };
+					}
+
+					if (const UAudioEngineSubsystem* Subsystem = Cast<UAudioEngineSubsystem>(Subscriber.GetObject()))
+					{
+						return Subsystem->GetAudioDeviceId();
+					}
+
+					Audio::FDeviceId DeviceId = INDEX_NONE;
+					if (TryGetDeviceIdFromWorldObject(Subscriber.GetObject(), DeviceId))
+					{
+						return DeviceId;
+					}
+
+					return { };
+				};
+
+				for (TScriptInterface<Core::ISubsonicEventSubscriberInterface>& EventSubscriber : EventSubscribers)
+				{
+					const TOptional<Audio::FDeviceId> SubscriberDeviceId = GetSubscriberDeviceId(EventSubscriber);
+					if (SubscriberDeviceId == DeviceId && EventSubscriber.GetObject() != nullptr)
+					{
+						Func(EventSubscriber);
+					}
+				}
+			}
+
+		public:
+			virtual void OnCollectionRegistered(const Core::FCollectionHandle& InHandle, Audio::FDeviceId DeviceId) override
+			{
+				IterateSubscribersMatchingDevice(DeviceId, [&InHandle](TScriptInterface<Core::ISubsonicEventSubscriberInterface>& Subscriber)
+				{
+					Subscriber->OnCollectionRegistered(InHandle);
+				});
+			}
+
+			virtual void OnCollectionUnregistered(const Core::FCollectionHandle& InHandle, Audio::FDeviceId DeviceId) override
+			{
+				IterateSubscribersMatchingDevice(DeviceId, [&InHandle](TScriptInterface<Core::ISubsonicEventSubscriberInterface>& Subscriber)
+				{
+					Subscriber->OnCollectionUnregistered(InHandle);
+				});
+			}
+
+			virtual void OnEventPreExecute(const Core::FSubsonicExecutor& InExecutor, const Core::FEventHandle& InEventHandle) override
+			{
+				IterateSubscribersMatchingDevice(InExecutor.GetDeviceId(), [&InExecutor, &InEventHandle](TScriptInterface<Core::ISubsonicEventSubscriberInterface>& Subscriber)
+				{
+					Subscriber->OnEventPreExecute(InExecutor, InEventHandle);
+				});
+			}
+
+			virtual void OnEventPostExecute(const Core::FSubsonicExecutor& InExecutor, const Core::FEventHandle& InEventHandle) override
+			{
+				IterateSubscribersMatchingDevice(InExecutor.GetDeviceId(), [&InExecutor, &InEventHandle](TScriptInterface<Core::ISubsonicEventSubscriberInterface>& Subscriber)
+				{
+					Subscriber->OnEventPostExecute(InExecutor, InEventHandle);
+				});
+			}
+
+			virtual void OnEventSubscriberRegistered(Core::ISubsonicEventSubscriberInterface& EventSubscriber) override
+			{
+				TScriptInterface<Core::ISubsonicEventSubscriberInterface> ScriptInterface;
+				ScriptInterface.SetObject(Cast<UObject>(&EventSubscriber));
+				ScriptInterface.SetInterface(&EventSubscriber);
+				EventSubscribers.AddUnique(ScriptInterface);
+			}
+
+			virtual void OnEventSubscriberUnregistered(Core::ISubsonicEventSubscriberInterface& EventSubscriber) override
+			{
+				TScriptInterface<Core::ISubsonicEventSubscriberInterface> ScriptInterface;
+				ScriptInterface.SetObject(Cast<UObject>(&EventSubscriber));
+				ScriptInterface.SetInterface(&EventSubscriber);
+				EventSubscribers.Remove(ScriptInterface);
+			}
+
+			virtual void OnExecutorRegistered(const Core::FSubsonicExecutor& InExecutor) override
+			{
+				IterateSubscribersMatchingDevice(InExecutor.GetDeviceId(), [&](TScriptInterface<Core::ISubsonicEventSubscriberInterface>& Subscriber)
+				{
+					Subscriber->OnExecutorRegistered(InExecutor);
+				});
+			}
+
+			virtual void OnExecutorUnregistered(const Core::FSubsonicExecutor& InExecutor) override
+			{
+				IterateSubscribersMatchingDevice(InExecutor.GetDeviceId(), [&](TScriptInterface<Core::ISubsonicEventSubscriberInterface>& Subscriber)
+				{
+					Subscriber->OnExecutorUnregistered(InExecutor);
+				});
+			}
+
+		private:
+			virtual void AddReferencedObjects(FReferenceCollector& Collector) override
+			{
+				TArray<TObjectPtr<UObject>> Objects;
+				Algo::Transform(EventSubscribers, Objects, [](const TScriptInterface<Core::ISubsonicEventSubscriberInterface>& Subscriber)
+				{
+					return TObjectPtr<UObject>(Subscriber.GetObject());
+				});
+				Collector.AddReferencedObjects(Objects);
+			}
+
+			virtual FString GetReferencerName() const override
+			{
+				return "FSubsonicEventRegistry";
+			}
+
+			TArray<TScriptInterface<Core::ISubsonicEventSubscriberInterface>> EventSubscribers;
+		};
+
+	} // namespace SubsystemPrivate
+
+	void USubsonicSubsystem::Initialize(FSubsystemCollectionBase& Collection)
+	{
+		Super::Initialize(Collection);
+		Core::ISubsonicEventRegistry::Initialize(TUniquePtr<Core::ISubsonicEventRegistry>(new SubsystemPrivate::FSubsonicEventRegistry()));
+	}
+
+	void USubsonicSubsystem::Deinitialize()
+	{
+		Super::Deinitialize();
+		Core::ISubsonicEventRegistry::Deinitialize();
+	}
+
+	USubsonicEventCollectionExecutor* USubsonicSubsystem::CreateExecutorBP(UObject* WorldContextObject, FName Name, const USubsonicEventCollection* InCollection)
+	{
+		if (InCollection && WorldContextObject)
+		{
+			Audio::FDeviceId DeviceId = INDEX_NONE;
+			if (SubsystemPrivate::TryGetDeviceIdFromWorldObject(WorldContextObject, DeviceId))
+			{
+				return USubsonicEventCollectionExecutor::Create(*WorldContextObject, Name, *InCollection, DeviceId);
+			}
+		}
+
+		return nullptr;
+	}
+} // namespace UE::Subsonic
diff --git a/Engine/Plugins/Experimental/Subsonic/Source/SubsonicEngine/Public/StandardActions/SubsonicAction_GeneratorSource.h b/Engine/Plugins/Experimental/Subsonic/Source/SubsonicEngine/Public/StandardActions/SubsonicAction_GeneratorSource.h
new file mode 100644
--- /dev/null
+++ b/Engine/Plugins/Experimental/Subsonic/Source/SubsonicEngine/Public/StandardActions/SubsonicAction_GeneratorSource.h
@@ -0,0 +1,63 @@
+// Copyright Epic Games, Inc. All Rights Reserved.
+#pragma once
+
+#include "Misc/CoreMiscDefines.h"
+#include "Sound/SoundWave.h"
+#include "SubsonicEventCollection.h"
+#include "SubsonicExecutor.h"
+#include "SubsonicParameterStore.h"
+
+#include "SubsonicAction_GeneratorSource.generated.h"
+
+namespace UE::Subsonic
+{
+	USTRUCT(MinimalAPI, BlueprintType, DisplayName = "Play Sound (GeneratorSource)")
+		struct FSubsonicEventAction_GeneratorSourcePlay : public Core::FSubsonicEventActionBase
+	{
+		GENERATED_BODY()
+
+		virtual ~FSubsonicEventAction_GeneratorSourcePlay() = default;
+
+		virtual void Execute(const Core::FSubsonicExecutor& InExecutor, const Core::FActionHandle& InHandle) const override;
+
+	public:
+		// Addressable name of Generator source instance.
+		UPROPERTY(EditAnywhere, Category = "Subsonic")
+		FName Name = "Sound01";
+
+		// Sound to play.
+		UPROPERTY(EditAnywhere, Category = "Subsonic")
+		TObjectPtr<USoundWave> Sound;
+
+		// Scope of source accessor.  If set to event, provided name is used
+		// to address the same source within the event, collection the same source
+		// within the collection, etc.
+		UPROPERTY(EditAnywhere, Category = "Subsonic")
+		ESubsonicExecutionScope Scope = ESubsonicExecutionScope::Collection;
+
+		// Authored parameters applied to the generator at play time (lowest priority layer).
+		UPROPERTY(EditAnywhere, Category = "Parameters")
+		FSubsonicParameterStore Parameters;
+	};
+
+	USTRUCT(MinimalAPI, BlueprintType, DisplayName = "Stop Sound (GeneratorSource)")
+	struct FSubsonicEventAction_GeneratorSourceStop : public Core::FSubsonicEventActionBase
+	{
+		GENERATED_BODY()
+
+		virtual ~FSubsonicEventAction_GeneratorSourceStop() = default;
+
+		virtual void Execute(const Core::FSubsonicExecutor& InExecutor, const Core::FActionHandle& InHandle) const override;
+
+	public:
+		// Addressable name of source instance.
+		UPROPERTY(EditAnywhere, Category = "Subsonic")
+		FName Name = "Sound01";
+
+		// Scope of source accessor.  If set to event, provided name is used
+		// to address the same source within the event, collection the same source
+		// within the collection, etc.
+		UPROPERTY(EditAnywhere, Category = "Subsonic")
+		ESubsonicExecutionScope Scope = ESubsonicExecutionScope::Collection;
+	};
+} // namespace UE::Subsonic
\ No newline at end of file
diff --git a/Engine/Plugins/Experimental/Subsonic/Source/SubsonicEngine/Public/StandardEventSubscribers/SubsonicEventSubscriberBase.h b/Engine/Plugins/Experimental/Subsonic/Source/SubsonicEngine/Public/StandardEventSubscribers/SubsonicEventSubscriberBase.h
new file mode 100644
--- /dev/null
+++ b/Engine/Plugins/Experimental/Subsonic/Source/SubsonicEngine/Public/StandardEventSubscribers/SubsonicEventSubscriberBase.h
@@ -0,0 +1,27 @@
+// Copyright Epic Games, Inc. All Rights Reserved.
+#pragma once
+
+#include "SubsonicEventSubscriberInterface.h"
+#include "Subsystems/AudioEngineSubsystem.h"
+
+#include "SubsonicEventSubscriberBase.generated.h"
+
+#define UE_API SUBSONICENGINE_API
+
+
+namespace UE::Subsonic
+{
+	// Forward Declarations
+	class USubsonicSubsystem;
+
+	UCLASS(MinimalAPI, Abstract)
+		class USubsonicEventSubscriberBase : public UAudioEngineSubsystem, public Core::ISubsonicEventSubscriberInterface
+	{
+		GENERATED_BODY()
+
+	public:
+		void Initialize(FSubsystemCollectionBase& Collection);
+		void Deinitialize();
+	};
+} // namespace UE::Subsonic
+#undef UE_API // SUBSONICENGINE_API
\ No newline at end of file
diff --git a/Engine/Plugins/Experimental/Subsonic/Source/SubsonicEngine/Public/StandardEventSubscribers/SubsonicGeneratorSourceSubscriber.h b/Engine/Plugins/Experimental/Subsonic/Source/SubsonicEngine/Public/StandardEventSubscribers/SubsonicGeneratorSourceSubscriber.h
new file mode 100644
--- /dev/null
+++ b/Engine/Plugins/Experimental/Subsonic/Source/SubsonicEngine/Public/StandardEventSubscribers/SubsonicGeneratorSourceSubscriber.h
@@ -0,0 +1,106 @@
+// Copyright Epic Games, Inc. All Rights Reserved.
+#pragma once
+
+#include "IAudioMixerGeneratorSource.h"
+#include "StandardEventSubscribers/SubsonicEventSubscriberBase.h"
+#include "StandardEventSubscribers/SubsonicRelay.h"
+#include "SubsonicHandles.h"
+#include "SubsonicParameterStore.h"
+#include "Tickable.h"
+
+#include "SubsonicGeneratorSourceSubscriber.generated.h"
+
+// Forward Declarations
+class USoundWave;
+
+namespace UE::Subsonic
+{
+	class FSubsonicGenerator;
+
+	// A source that is no longer addressable from actions but still playing.
+	// Paired with its generator so ClearStaleSounds can check completion.
+	struct FReleasedSource
+	{
+		TUniquePtr<Audio::IAudioMixerGeneratorSource> Source;
+		TSharedPtr<FSubsonicGenerator> Generator;
+	};
+
+	struct FInstanceSources
+	{
+		// Currently addressable generator source via handle
+		TUniquePtr<Audio::IAudioMixerGeneratorSource> ActiveSource;
+
+		// Sources no longer addressable from actions but actively playing.
+		TArray<FReleasedSource> ReleasedSources;
+
+		// The active generator, shared with the relay command queue so that
+		// SetParameters commands can target it directly on the render thread.
+		TSharedPtr<FSubsonicGenerator> Generator;
+
+		FSubsonicParameterStore AuthoredParams;   // from action definition (lowest priority)
+		FSubsonicParameterStore RuntimeParams;    // from post-playback SetParameters (highest priority)
+
+		// Merges all three layers (authored < trigger < runtime) and enqueues a
+		// SetParameters command into the relay's global command queue.
+		void ApplyParameters(FSubsonicRelay& Relay);
+	};
+
+	using FGeneratorSourceCache = TMap<FName, FInstanceSources>;
+
+	/**
+	 * Interface to GeneratorSource playback Subsonic implementation.
+	 */
+	UCLASS()
+	class USubsonicGeneratorSourceSubscriber final : public USubsonicEventSubscriberBase, public FTickableGameObject
+	{
+		GENERATED_BODY()
+
+	public:
+		using FActionHandle = Core::FActionHandle;
+		using FEventHandle = Core::FEventHandle;
+		using FCollectionHandle = Core::FCollectionHandle;
+
+		USubsonicGeneratorSourceSubscriber();
+
+		//~ Begin USubsystem interface
+		virtual void Initialize(FSubsystemCollectionBase& Collection) override;
+		virtual void Deinitialize() override;
+		//~ End USubsystem interface
+
+		virtual void OnCollectionRegistered(const FCollectionHandle& InEvent) override;
+		virtual void OnCollectionUnregistered(const FCollectionHandle& InEvent) override;
+
+		void PlaySound(FName Name, USoundWave& Sound, const FSubsonicParameterStore& Params, bool bReleaseExisting = true);
+		void PlaySound(const FCollectionHandle& InHandle, FName Name, USoundWave& Sound, const FSubsonicParameterStore& Params, bool bReleaseExisting = true);
+		void PlaySound(const FEventHandle& InHandle, FName Name, USoundWave& Sound, const FSubsonicParameterStore& Params, bool bReleaseExisting = true);
+		void PlaySound(const FActionHandle& InHandle, FName Name, USoundWave& Sound, const FSubsonicParameterStore& Params, bool bReleaseExisting = true);
+
+		void StopSound(FName Name);
+		void StopSound(const FCollectionHandle& InHandle, FName Name);
+		void StopSound(const FEventHandle& InHandle, FName Name);
+		void StopSound(const FActionHandle& InHandle, FName Name);
+
+		void SetParameters(FName SourceName, const FSubsonicParameterStore& Params);
+		void SetParameters(const FCollectionHandle& InHandle, FName SourceName, const FSubsonicParameterStore& Params);
+		void SetParameters(const FEventHandle& InHandle, FName SourceName, const FSubsonicParameterStore& Params);
+		void SetParameters(const FActionHandle& InHandle, FName SourceName, const FSubsonicParameterStore& Params);
+
+		SUBSONICENGINE_API void SetAuthoredParameters(FName SourceName, const FSubsonicParameterStore& Params);
+		SUBSONICENGINE_API void SetAuthoredParameters(const FCollectionHandle& InHandle, FName SourceName, const FSubsonicParameterStore& Params);
+		SUBSONICENGINE_API void SetAuthoredParameters(const FEventHandle& InHandle, FName SourceName, const FSubsonicParameterStore& Params);
+		SUBSONICENGINE_API void SetAuthoredParameters(const FActionHandle& InHandle, FName SourceName, const FSubsonicParameterStore& Params);
+
+		// FTickableGameObject interface
+		virtual void Tick(float DeltaTime) override;
+		virtual bool IsTickableInEditor() const override { return true; }
+		virtual TStatId GetStatId() const override;
+
+	private:
+		TSharedRef<FSubsonicRelay> Relay;
+
+		TMap<int32, FGeneratorSourceCache> GlobalSources;
+		TMap<FCollectionHandle, FGeneratorSourceCache> CollectionSources;
+		TMap<FEventHandle, FGeneratorSourceCache> EventSources;
+		TMap<FActionHandle, FGeneratorSourceCache> ActionSources;
+	};
+} // namespace UE::Subsonic
diff --git a/Engine/Plugins/Experimental/Subsonic/Source/SubsonicEngine/Public/StandardEventSubscribers/SubsonicRelay.h b/Engine/Plugins/Experimental/Subsonic/Source/SubsonicEngine/Public/StandardEventSubscribers/SubsonicRelay.h
new file mode 100644
--- /dev/null
+++ b/Engine/Plugins/Experimental/Subsonic/Source/SubsonicEngine/Public/StandardEventSubscribers/SubsonicRelay.h
@@ -0,0 +1,42 @@
+// Copyright Epic Games, Inc. All Rights Reserved.
+
+#pragma once
+
+#include "Containers/Array.h"
+#include "HAL/CriticalSection.h"
+#include "SubsonicParameterStore.h"
+#include "Templates/SharedPointer.h"
+
+namespace Audio { class FMixerDevice; }
+
+namespace UE::Subsonic
+{
+	class FSubsonicGenerator;
+
+	// Typed command for the per-device command queue on FSubsonicRelay. Commands are expected to be added on the game thread but can
+	//  be safely added from other threads. All commands in one game tick are pushed to the audio render thread in a single batch.
+	struct FRelayCommand
+	{
+		enum class EType : uint8 { SetParameters, Play, Stop };
+
+		EType Type = EType::SetParameters;
+		TSharedPtr<FSubsonicParameterStore> Params;      // SetParameters only
+		TSharedPtr<FSubsonicGenerator> Target;            // all types
+	};
+
+	// Per-device object responsible for communicating state changes to the audio render thread.
+	class FSubsonicRelay
+	{
+	public:
+		// Append a command to the pending queue (lock-protected).
+		void EnqueueCommand(FRelayCommand&& Cmd);
+
+		// Game thread (called from subscriber Tick): flush pending commands into a single
+		// AudioRenderThreadCommand lambda on the mixer device.
+		void Tick(Audio::FMixerDevice& MixerDevice);
+
+	private:
+		FCriticalSection CommandLock;
+		TArray<FRelayCommand> PendingCommands;
+	};
+} // namespace UE::Subsonic
diff --git a/Engine/Plugins/Experimental/Subsonic/Source/SubsonicEngine/Public/SubsonicDeviceUtils.h b/Engine/Plugins/Experimental/Subsonic/Source/SubsonicEngine/Public/SubsonicDeviceUtils.h
new file mode 100644
--- /dev/null
+++ b/Engine/Plugins/Experimental/Subsonic/Source/SubsonicEngine/Public/SubsonicDeviceUtils.h
@@ -0,0 +1,35 @@
+// Copyright Epic Games, Inc. All Rights Reserved.
+#pragma once
+
+#include "AudioDevice.h"
+#include "AudioDeviceManager.h"
+#include "SubsonicExecutor.h"
+
+namespace UE::Subsonic
+{
+	// Utility to find an audio device subsystem from a scope key.
+	// Extracted from TScopeKey to keep SubsonicCore free of AudioDevice dependencies.
+	template <typename SubsystemType, typename ScopeHandleType>
+	SubsystemType* FindDeviceSubsystem(const Core::TScopeKey<ScopeHandleType>& InKey)
+	{
+		const FAudioDevice* AudioDevice = nullptr;
+		if (const FAudioDeviceManager* DeviceManager = FAudioDeviceManager::Get())
+		{
+			if (InKey.DeviceId == INDEX_NONE)
+			{
+				AudioDevice = DeviceManager->GetMainAudioDeviceRaw();
+			}
+			else
+			{
+				AudioDevice = DeviceManager->GetAudioDeviceRaw(InKey.DeviceId);
+			}
+		}
+
+		if (AudioDevice)
+		{
+			return AudioDevice->GetSubsystem<SubsystemType>();
+		}
+
+		return nullptr;
+	}
+} // namespace UE::Subsonic
diff --git a/Engine/Plugins/Experimental/Subsonic/Source/SubsonicEngine/Public/SubsonicEventCollectionObjects.h b/Engine/Plugins/Experimental/Subsonic/Source/SubsonicEngine/Public/SubsonicEventCollectionObjects.h
new file mode 100644
--- /dev/null
+++ b/Engine/Plugins/Experimental/Subsonic/Source/SubsonicEngine/Public/SubsonicEventCollectionObjects.h
@@ -0,0 +1,99 @@
+// Copyright Epic Games, Inc. All Rights Reserved.
+#pragma once
+
+#include "SubsonicEventCollection.h"
+#include "SubsonicExecutor.h"
+#include "SubsonicHandles.h"
+#include "UObject/Object.h"
+
+#include "SubsonicEventCollectionObjects.generated.h"
+
+#define UE_API SUBSONICENGINE_API
+
+namespace UE::Subsonic
+{
+	UCLASS(MinimalAPI, BlueprintType)
+	class USubsonicEventCollection final : public UObject
+	{
+		GENERATED_BODY()
+
+	public:
+		UE_API const Core::FSubsonicEventCollectionDefinition& GetDefinition() const;
+
+#if WITH_EDITOR
+		UE_API static FName GetDefinitionPropertyName();
+#endif // WITH_EDITOR
+
+		UE_API Core::FCollectionHandle GetHandle() const;
+
+#if WITH_EDITOR
+		UE_API UE_INTERNAL Core::FSubsonicEventCollectionDefinition& GetMutableDefinition();
+#endif // WITH_EDITOR
+
+		UE_API void SetDefinition(Core::FSubsonicEventCollectionDefinition&& Definition);
+
+	private:
+		UPROPERTY(EditAnywhere, Category = "Subsonic", meta = (ShowOnlyInnerProperties))
+		Core::FSubsonicEventCollectionDefinition CollectionDefinition;
+
+		virtual void PostInitProperties() override;
+		virtual void PostLoad() override;
+		virtual void BeginDestroy() override;
+
+#if WITH_EDITOR
+		virtual void PostEditChangeChainProperty(FPropertyChangedChainEvent& PropertyChangedEvent) override;
+		virtual void PostEditUndo() override;
+#endif // WITH_EDITOR
+	};
+
+	UENUM(BlueprintType, meta = (DisplayName = "Execution Result"))
+	enum class ESubsonicExecutionResult : uint8
+	{
+		Succeeded,
+		Failed
+	};
+
+	// Executor handle for a Subsonic Event Collection.
+	UCLASS(MinimalAPI, BlueprintType)
+	class USubsonicEventCollectionExecutor final : public UObject
+	{
+		GENERATED_BODY()
+
+	public:
+		virtual void BeginDestroy() override;
+
+		// Creates a transient executor that is bound to the given collection.
+		UE_API static USubsonicEventCollectionExecutor* Create(UObject& Parent, FName Name, const USubsonicEventCollection& InCollection, Audio::FDeviceId DeviceId);
+
+		// Executes an event with the given gameplay tag if it exists.
+		// Sets OutResult to Succeeded if event was found and executed,
+		// Failed if not (enum provided as execution expansion support).
+		UFUNCTION(BlueprintCallable, Category = "Subsonic|Collection|Events", meta = (ExpandEnumAsExecs = "OutResult"))
+		UE_API void ExecuteEvent(FGameplayTag EventTag, ESubsonicExecutionResult& OutResult);
+
+		// Returns collection executor is bound to.
+		UE_API const USubsonicEventCollection* GetCollection() const;
+
+		// Returns collection executor is bound to and asserts if no collection is set.
+		UE_API const USubsonicEventCollection& GetCollectionChecked() const;
+
+		// Returns underlying executor UObject wraps.
+		UE_API const Core::FSubsonicExecutor& GetExecutor() const;
+
+		// Returns whether this executor is valid (has a valid ID and is bound to an event collection)
+		UE_API bool IsValid() const;
+
+		// Unregisters the underlying executor implementation, clearing
+		// any state pertaining to the given executor managed by subscribers
+		// and releasing association with the provided collection when created.
+		UE_API void Unregister();
+
+	private:
+		UPROPERTY(Transient, BlueprintReadOnly, Category = "Subsonic", meta = (AllowPrivateAccess))
+		TObjectPtr<const USubsonicEventCollection> Collection;
+
+		Core::FSubsonicExecutor Executor = Core::FSubsonicExecutor::GetInvalid();
+	};
+} // namespace UE::Subsonic
+
+#undef UE_API // SUBSONICENGINE_API
diff --git a/Engine/Plugins/Experimental/Subsonic/Source/SubsonicEngine/Public/SubsonicSubsystem.h b/Engine/Plugins/Experimental/Subsonic/Source/SubsonicEngine/Public/SubsonicSubsystem.h
new file mode 100644
--- /dev/null
+++ b/Engine/Plugins/Experimental/Subsonic/Source/SubsonicEngine/Public/SubsonicSubsystem.h
@@ -0,0 +1,33 @@
+// Copyright Epic Games, Inc. All Rights Reserved.
+
+#pragma once
+
+#include "SubsonicEventCollectionObjects.h"
+#include "Subsystems/EngineSubsystem.h"
+
+#include "SubsonicSubsystem.generated.h"
+
+#define UE_API SUBSONICENGINE_API
+
+
+namespace UE::Subsonic
+{
+	/**
+	 * Interface to the Subsonic plugin.
+	 */
+	UCLASS(MinimalAPI, BlueprintType)
+	class USubsonicSubsystem final : public UEngineSubsystem
+	{
+		GENERATED_BODY()
+
+	public:
+		//~ Begin USubsystem interface
+		UE_API virtual void Initialize(FSubsystemCollectionBase& Collection) override;
+		UE_API virtual void Deinitialize() override;
+		//~ End USubsystem interface
+
+		UFUNCTION(BlueprintCallable, Category = "Subsonic", meta = (DisplayName = "Create Executor", WorldContext = "WorldContextObject"))
+		UE_API UPARAM(DisplayName = "Executor") USubsonicEventCollectionExecutor* CreateExecutorBP(UObject* WorldContextObject, FName Name, const USubsonicEventCollection* Collection);
+	};
+} // namespace UE::Subsonic
+#undef UE_API
diff --git a/Engine/Plugins/Experimental/Subsonic/Source/SubsonicEngine/SubsonicEngine.Build.cs b/Engine/Plugins/Experimental/Subsonic/Source/SubsonicEngine/SubsonicEngine.Build.cs
new file mode 100644
--- /dev/null
+++ b/Engine/Plugins/Experimental/Subsonic/Source/SubsonicEngine/SubsonicEngine.Build.cs
@@ -0,0 +1,39 @@
+// Copyright Epic Games, Inc. All Rights Reserved.
+
+using UnrealBuildTool;
+
+public class SubsonicEngine : ModuleRules
+{
+	public SubsonicEngine(ReadOnlyTargetRules Target) : base(Target)
+	{
+		PCHUsage = ModuleRules.PCHUsageMode.UseExplicitOrSharedPCHs;
+
+		PublicDependencyModuleNames.AddRange(
+			new string[]
+			{
+				"AudioMixerCore",
+				"Core",
+				"CoreUObject",
+				"Engine",
+				"GameplayTags",
+				"SubsonicCore",
+			}
+		);
+
+		PrivateDependencyModuleNames.AddRange(
+			new string[]
+			{
+				"AudioExtensions",
+				"AudioMixer",
+				"MetasoundEngine",
+				"MetasoundFrontend",
+				"MetasoundGenerator",
+				"MetasoundGraphCore",
+				"PropertyBindingUtils",
+				"SignalProcessing",
+			}
+		);
+
+		bAllowUETypesInNamespaces = true;
+	}
+}
diff --git a/Engine/Plugins/Experimental/Subsonic/Source/SubsonicEngineTest/Private/SubsonicEngineTestModule.cpp b/Engine/Plugins/Experimental/Subsonic/Source/SubsonicEngineTest/Private/SubsonicEngineTestModule.cpp
new file mode 100644
--- /dev/null
+++ b/Engine/Plugins/Experimental/Subsonic/Source/SubsonicEngineTest/Private/SubsonicEngineTestModule.cpp
@@ -0,0 +1,11 @@
+// Copyright Epic Games, Inc. All Rights Reserved.
+
+#include "Modules/ModuleInterface.h"
+#include "Modules/ModuleManager.h"
+
+
+class FSubsonicEngineTestModule : public IModuleInterface
+{
+};
+
+IMPLEMENT_MODULE(FSubsonicEngineTestModule, SubsonicEngineTest)
diff --git a/Engine/Plugins/Experimental/Subsonic/Source/SubsonicEngineTest/Private/SubsonicEventCollectionObjectsTest.cpp b/Engine/Plugins/Experimental/Subsonic/Source/SubsonicEngineTest/Private/SubsonicEventCollectionObjectsTest.cpp
new file mode 100644
--- /dev/null
+++ b/Engine/Plugins/Experimental/Subsonic/Source/SubsonicEngineTest/Private/SubsonicEventCollectionObjectsTest.cpp
@@ -0,0 +1,66 @@
+// Copyright Epic Games, Inc. All Rights Reserved.
+
+#include "CQTest.h"
+#include "GameplayTagContainer.h"
+#include "SubsonicEventCollection.h"
+#include "SubsonicEventCollectionObjects.h"
+
+
+namespace UE::Subsonic::Test
+{
+	namespace TestPrivate
+	{
+		// Creates a transient USubsonicEventCollection with a single no-op event registered
+		// under the provdied EventTag. Returns null if creation or event registration fails.
+		USubsonicEventCollection* CreateCollectionWithEvent(const FGameplayTag EventTag)
+		{
+			USubsonicEventCollection* Collection = NewObject<USubsonicEventCollection>(GetTransientPackage());
+			if (!Collection)
+			{
+				return nullptr;
+			}
+
+			if (!Collection->GetMutableDefinition().AddEvent(EventTag))
+			{
+				return nullptr;
+			}
+
+			return Collection;
+		}
+	} // namespace TestPrivate
+
+	TEST_CLASS_WITH_FLAGS(SubsonicEventCollectionObjectsTest, "Audio.Subsonic.Engine",
+		EAutomationTestFlags::EditorContext | EAutomationTestFlags::ProductFilter)
+	{
+		TEST_METHOD(CreateCollectionWithNoOpEvent)
+		{
+			USubsonicEventCollection* Collection = TestPrivate::CreateCollectionWithEvent(Core::TAG_SubsonicCore_Event_Play);
+			ASSERT_THAT(IsNotNull(Collection));
+
+			if (Collection)
+			{
+				ASSERT_THAT(IsTrue(Collection->GetDefinition().ContainsEvent(Core::TAG_SubsonicCore_Event_Play)));
+			}
+		}
+
+		TEST_METHOD(CreateExecutorAndFireNoOpEvent)
+		{
+			USubsonicEventCollection* Collection = TestPrivate::CreateCollectionWithEvent(Core::TAG_SubsonicCore_Event_Play);
+			ASSERT_THAT(IsNotNull(Collection));
+
+			if (Collection)
+			{
+				USubsonicEventCollectionExecutor* Executor = USubsonicEventCollectionExecutor::Create(
+					*GetTransientPackage(), FName("TestExecutor"), *Collection, INDEX_NONE);
+				ASSERT_THAT(IsNotNull(Executor));
+				
+				if (Executor)
+				{
+					ASSERT_THAT(IsTrue(Executor->IsValid()));
+					const bool bExecuted = Executor->GetExecutor().ExecuteEvent(Core::TAG_SubsonicCore_Event_Play.GetTag().GetTagName());
+					ASSERT_THAT(IsTrue(bExecuted));
+				}
+			}
+		}
+	};
+} // namespace UE::Subsonic::Test
diff --git a/Engine/Plugins/Experimental/Subsonic/Source/SubsonicEngineTest/SubsonicEngineTest.Build.cs b/Engine/Plugins/Experimental/Subsonic/Source/SubsonicEngineTest/SubsonicEngineTest.Build.cs
new file mode 100644
--- /dev/null
+++ b/Engine/Plugins/Experimental/Subsonic/Source/SubsonicEngineTest/SubsonicEngineTest.Build.cs
@@ -0,0 +1,27 @@
+// Copyright Epic Games, Inc. All Rights Reserved.
+
+using UnrealBuildTool;
+
+public class SubsonicEngineTest : ModuleRules
+{
+	public SubsonicEngineTest(ReadOnlyTargetRules Target) : base(Target)
+	{
+		PCHUsage = ModuleRules.PCHUsageMode.UseExplicitOrSharedPCHs;
+		OptimizeCode = CodeOptimization.Never;
+
+		PrivateDependencyModuleNames.AddRange(
+			new string[]
+			{
+				"Core",
+				"CoreUObject",
+				"CQTest",
+				"Engine",
+				"GameplayTags",
+				"SubsonicCore",
+				"SubsonicEngine"
+			}
+		);
+
+		bAllowUETypesInNamespaces = true;
+	}
+}
diff --git a/Engine/Plugins/Experimental/Subsonic/Subsonic.uplugin b/Engine/Plugins/Experimental/Subsonic/Subsonic.uplugin
--- a/Engine/Plugins/Experimental/Subsonic/Subsonic.uplugin
+++ b/Engine/Plugins/Experimental/Subsonic/Subsonic.uplugin
@@ -17,7 +17,12 @@
 	"EnabledByDefault": false,
 	"Modules": [
 		{
-			"Name": "Subsonic",
+			"Name": "SubsonicCore",
+			"Type": "Runtime",
+			"LoadingPhase": "Default"
+		},
+		{
+			"Name": "SubsonicEngine",
 			"Type": "Runtime",
 			"LoadingPhase": "Default"
 		},
@@ -25,6 +30,11 @@
 			"Name": "SubsonicEditor",
 			"Type": "UncookedOnly",
 			"LoadingPhase": "Default"
+		},
+		{
+			"Name": "SubsonicEngineTest",
+			"Type": "UncookedOnly",
+			"LoadingPhase": "Default"
 		}
 	],
 	"Plugins": [
