Remove code optimization submitted by mistake
#rnx

diff --git a/Engine/Plugins/Experimental/Subsonic/Source/SubsonicEngineTest/SubsonicEngineTest.Build.cs b/Engine/Plugins/Experimental/Subsonic/Source/SubsonicEngineTest/SubsonicEngineTest.Build.cs
--- a/Engine/Plugins/Experimental/Subsonic/Source/SubsonicEngineTest/SubsonicEngineTest.Build.cs
+++ b/Engine/Plugins/Experimental/Subsonic/Source/SubsonicEngineTest/SubsonicEngineTest.Build.cs
@@ -7,7 +7,6 @@ public class SubsonicEngineTest : ModuleRules
 	public SubsonicEngineTest(ReadOnlyTargetRules Target) : base(Target)
 	{
 		PCHUsage = ModuleRules.PCHUsageMode.UseExplicitOrSharedPCHs;
-		OptimizeCode = CodeOptimization.Never;
 
 		PrivateDependencyModuleNames.AddRange(
 			new string[]
