Subsonic: Subscriber implementation consolidation and removal of action and event scope

#rnx
#Synced-CL 52327832

diff --git a/Engine/Plugins/Experimental/Subsonic/Source/SubsonicCore/Private/SubsonicEventCollection.cpp b/Engine/Plugins/Experimental/Subsonic/Source/SubsonicCore/Private/SubsonicEventCollection.cpp
--- a/Engine/Plugins/Experimental/Subsonic/Source/SubsonicCore/Private/SubsonicEventCollection.cpp
+++ b/Engine/Plugins/Experimental/Subsonic/Source/SubsonicCore/Private/SubsonicEventCollection.cpp
@@ -42,6 +42,61 @@ namespace UE::Subsonic::Core
 		return GET_MEMBER_NAME_CHECKED(FSubsonicEventActionBase, Bindings);
 	}
 
+	FName FSubsonicEventActionBase::GenerateDefaultBaseIdentifier(const FEventHandle& OwnerHandle)
+	{
+		const bool bIsPlayEvent = OwnerHandle.EventName == TAG_SubsonicCore_Event_Play.GetTag().GetTagName();
+		const bool bIsStopEvent = OwnerHandle.EventName == TAG_SubsonicCore_Event_Stop.GetTag().GetTagName();
+		FName DefaultName = bIsPlayEvent || bIsStopEvent || OwnerHandle.EventName.IsNone()
+			? OwnerHandle.Collection.CollectionName
+			: OwnerHandle.EventName;
+		TArray<FString> Tokens;
+		DefaultName.ToString().ParseIntoArray(Tokens, TEXT("/"));
+		if (!Tokens.IsEmpty())
+		{
+			const FString Last = MoveTemp(Tokens.Last());
+			Last.ParseIntoArray(Tokens, TEXT("."));
+			DefaultName = Tokens.IsEmpty()
+				? *Last
+				: *Tokens.Last();
+		}
+		return DefaultName;
+	}
+
+	void FSubsonicEventActionBase::InitializeDefaultActionName(TInstancedStruct<FSubsonicEventActionBase>& ActionConfig, const FEventHandle& OwnerHandle)
+	{
+		if (FSubsonicEventActionBase* Action = ActionConfig.GetMutablePtr())
+		{
+			const FName BaseName = GenerateDefaultBaseIdentifier(OwnerHandle);
+			Action->InitializeDefaultName(ActionConfig.GetScriptStruct(), BaseName);
+		}
+	}
+
+	void FSubsonicEventActionBase::InitializeDefaultName(const UScriptStruct* ActionStruct, FName BaseName)
+	{
+		if (!ActionStruct)
+		{
+			return;
+		}
+
+		static const FLazyName Name("Name");
+		for (TFieldIterator<FProperty> PropIt(ActionStruct); PropIt; ++PropIt)
+		{
+			if (PropIt->GetOwnerStruct() != ActionStruct)
+			{
+				break;
+			}
+
+			if (PropIt->GetFName() == Name)
+			{
+				if (FNameProperty* NameProp = CastField<FNameProperty>(*PropIt))
+				{
+					NameProp->SetPropertyValue_InContainer(this, BaseName);
+				}
+				break;
+			}
+		}
+	}
+
 	FText FSubsonicEventActionBase::GetDisplayInfo() const
 	{
 		return FText();
@@ -211,8 +266,8 @@ namespace UE::Subsonic::Core
 
 	void FSubsonicEvent::Execute(const FSubsonicExecutor& InExecutor, const FEventHandle& InHandle) const
 	{
-		checkf(InExecutor.GetCollectionHandle().CollectionId == InHandle.Collection.CollectionId,
-			TEXT("Invalid execution context: executor CollectionId must match that of the given event handle"));
+		checkf(InExecutor.GetCollection() && InExecutor.GetCollection()->GetCollectionId() == InHandle.Collection.CollectionId,
+			TEXT("Invalid execution context: executor Collection must be valid and Id must match that of the given event handle"));
 
 		FActionHandle ActionHandle { .Event = InHandle };
 		for (int32 Index = 0; Index < ActionCollection.Num(); ++Index)
diff --git a/Engine/Plugins/Experimental/Subsonic/Source/SubsonicCore/Private/SubsonicExecutor.cpp b/Engine/Plugins/Experimental/Subsonic/Source/SubsonicCore/Private/SubsonicExecutor.cpp
--- a/Engine/Plugins/Experimental/Subsonic/Source/SubsonicCore/Private/SubsonicExecutor.cpp
+++ b/Engine/Plugins/Experimental/Subsonic/Source/SubsonicCore/Private/SubsonicExecutor.cpp
@@ -21,53 +21,30 @@ namespace UE::Subsonic::Core
 	{
 	}
 
-	FSubsonicExecutor::FSubsonicExecutor(FSubsonicExecutor&& InOther)
+	FSubsonicExecutor::~FSubsonicExecutor()
 	{
 		Unregister();
-
-		CollectionAccessor = MoveTemp(InOther.CollectionAccessor);
-		Parameters = MoveTemp(InOther.Parameters);
-
-		DeviceId = InOther.DeviceId;
-		InOther.DeviceId = ExecutorPrivate::InvalidDeviceId;
-
-		ExecutorId = InOther.ExecutorId;
-		InOther.ExecutorId = ExecutorPrivate::InvalidId;
 	}
 
-	FSubsonicExecutor::~FSubsonicExecutor()
+	TSharedRef<FSubsonicExecutor> FSubsonicExecutor::Create(Audio::FDeviceId InDeviceId, FSubsonicExecutor::FCollectionAccessorPtr&& CollectionPtr)
 	{
-		Unregister();
+		static uint32 MaxExecutorId = ExecutorPrivate::InvalidId;
+		TSharedRef<FSubsonicExecutor> NewExecutor = MakeShared<FSubsonicExecutor>();
+		NewExecutor->CollectionAccessor = MoveTemp(CollectionPtr);
+		NewExecutor->ExecutorId = ++MaxExecutorId;
+		NewExecutor->DeviceId = InDeviceId;
+		NewExecutor->Register();
+		return NewExecutor;
 	}
 
-	FSubsonicExecutor& FSubsonicExecutor::operator=(FSubsonicExecutor&& InOther)
+	const FSubsonicEventCollectionDefinition* FSubsonicExecutor::GetCollection() const
 	{
-		if (&InOther != this)
+		if (CollectionAccessor.IsValid())
 		{
-			Unregister();
-
-			CollectionAccessor = MoveTemp(InOther.CollectionAccessor);
-			Parameters = MoveTemp(InOther.Parameters);
-
-			DeviceId = InOther.DeviceId;
-			InOther.DeviceId = ExecutorPrivate::InvalidDeviceId;
-
-			ExecutorId = InOther.ExecutorId;
-			InOther.ExecutorId = ExecutorPrivate::InvalidId;
+			return CollectionAccessor->GetDefinition();
 		}
 
-		return *this;
-	}
-
-	FSubsonicExecutor FSubsonicExecutor::Create(Audio::FDeviceId InDeviceId, FSubsonicExecutor::FCollectionAccessorPtr&& CollectionPtr)
-	{
-		static uint32 MaxExecutorId = ExecutorPrivate::InvalidId;
-		FSubsonicExecutor NewHandle;
-		NewHandle.CollectionAccessor = MoveTemp(CollectionPtr);
-		NewHandle.ExecutorId = ++MaxExecutorId;
-		NewHandle.DeviceId = InDeviceId;
-		NewHandle.Register();
-		return NewHandle;
+		return nullptr;
 	}
 
 	FCollectionHandle FSubsonicExecutor::GetCollectionHandle() const
@@ -132,11 +109,6 @@ namespace UE::Subsonic::Core
 		return ExecutorPrivate::InvalidId;
 	}
 
-	FSubsonicExecutor FSubsonicExecutor::GetInvalid()
-	{
-		return { };
-	}
-
 	bool FSubsonicExecutor::IsValid() const
 	{
 		return ExecutorId != ExecutorPrivate::InvalidId;
diff --git a/Engine/Plugins/Experimental/Subsonic/Source/SubsonicCore/Public/SubsonicEventCollection.h b/Engine/Plugins/Experimental/Subsonic/Source/SubsonicCore/Public/SubsonicEventCollection.h
--- a/Engine/Plugins/Experimental/Subsonic/Source/SubsonicCore/Public/SubsonicEventCollection.h
+++ b/Engine/Plugins/Experimental/Subsonic/Source/SubsonicCore/Public/SubsonicEventCollection.h
@@ -46,9 +46,21 @@ namespace UE::Subsonic::Core
 		UE_INTERNAL UE_API const TMap<FName, FSubsonicEventActionBoundProperties>& GetBindings() const;
 		UE_INTERNAL UE_API static FName GetBindingsParameterName();
 
+		// Computes a reasonable default base identifiable name from the provided owning event handle.
+		UE_API static FName GenerateDefaultBaseIdentifier(const FEventHandle& OwnerHandle);
+
 		// Returns display info that is custom to the given action instance to be shown
 		// adjacent to the action, for example, in the Subsonic Event Editor action tree.
 		UE_API virtual FText GetDisplayInfo() const;
+
+		// Initializes default property names on a newly created action config.
+		// Computes a base name from context, then delegates to the virtual InitializeDefaultName.
+		UE_API static void InitializeDefaultActionName(TInstancedStruct<FSubsonicEventActionBase>& ActionConfig, const FEventHandle& OwnerHandle);
+
+		// Virtual hook for per-action-type customization of default property names.
+		// Called by InitializeDefaultActionName with the computed base name.
+		// Base implementation sets the first FNameProperty named "Name" to BaseName via reflection.
+		UE_API virtual void InitializeDefaultName(const UScriptStruct* ActionStruct, FName BaseName);
 #endif // WITH_EDITOR
 
 	public:
diff --git a/Engine/Plugins/Experimental/Subsonic/Source/SubsonicCore/Public/SubsonicExecutor.h b/Engine/Plugins/Experimental/Subsonic/Source/SubsonicCore/Public/SubsonicExecutor.h
--- a/Engine/Plugins/Experimental/Subsonic/Source/SubsonicCore/Public/SubsonicExecutor.h
+++ b/Engine/Plugins/Experimental/Subsonic/Source/SubsonicCore/Public/SubsonicExecutor.h
@@ -3,9 +3,10 @@
 
 #include "AudioDefines.h"
 #include "HAL/Platform.h"
-#include "Misc/CoreDefines.h"
 #include "SubsonicHandles.h"
 #include "SubsonicParameterStore.h"
+#include "Templates/SharedPointer.h"
+#include "Templates/SharedPointerInternals.h"
 #include "Templates/TypeHash.h"
 #include "Templates/UniquePtr.h"
 #include "UObject/NameTypes.h"
@@ -24,26 +25,30 @@ UENUM(BlueprintType)
 enum class ESubsonicExecutionScope : uint8
 {
 	Global,
-	Collection,
-	Event,
-	Action
+	Executor
 };
 
 
 namespace UE::Subsonic::Core
 {
-	struct FSubsonicExecutor final
+	// Execution object that tracks explicit execution context for any caller of events related to a given collection.
+	// RAII design that automatically registers and unregisters with the Subsonic Event Registry, and by extension,
+	// any implemented and active subscribers. This enables subscribers to track and manage data related to
+	// event requests from a given executor instance in a data-oriented fashion.
+	struct FSubsonicExecutor final : public TSharedFromThis<FSubsonicExecutor>
 	{
 	private:
+		template <typename ObjectType, ESPMode Mode>
+		friend class SharedPointerInternals::TIntrusiveReferenceController;
+
 		FSubsonicExecutor();
 		FSubsonicExecutor(const FSubsonicExecutor&) = delete;
-
-	public:
-		FSubsonicExecutor(FSubsonicExecutor&&);
+		FSubsonicExecutor(FSubsonicExecutor&&) = delete;
 		FSubsonicExecutor& operator=(const FSubsonicExecutor& InOther) = delete;
-		UE_API ~FSubsonicExecutor();
+		FSubsonicExecutor& operator=(FSubsonicExecutor&& InOther) = delete;
 
-		UE_API FSubsonicExecutor& operator=(FSubsonicExecutor&& InOther);
+	public:
+		virtual ~FSubsonicExecutor();
 
 		struct ICollectionAccessor
 		{
@@ -54,13 +59,13 @@ namespace UE::Subsonic::Core
 		};
 		using FCollectionAccessorPtr = TUniquePtr<ICollectionAccessor>;
 
-		UE_API static FSubsonicExecutor Create(Audio::FDeviceId InDeviceId, FCollectionAccessorPtr&& CollectionPtr);
+		UE_API static TSharedRef<FSubsonicExecutor> Create(Audio::FDeviceId InDeviceId, FCollectionAccessorPtr&& CollectionPtr);
 
 		UE_API bool ExecuteEvent(FName EventName) const;
 
 		UE_API static uint32 GetInvalidId();
-		UE_API static FSubsonicExecutor GetInvalid();
 
+		UE_API const FSubsonicEventCollectionDefinition* GetCollection() const;
 		UE_API FCollectionHandle GetCollectionHandle() const;
 		UE_API Audio::FDeviceId GetDeviceId() const;
 		UE_API uint32 GetId() const;
@@ -91,41 +96,47 @@ namespace UE::Subsonic::Core
 		FSubsonicParameterStore Parameters;
 	};
 
-	template <typename ScopeHandleType>
-	struct TScopeKey
+	// Key consisting of executor ID and Audio DeviceId (holds CollectionHandle in builds with logging
+	// enabled for debugging). Provides easily readable, hashable key while not providing direct executor 
+	// ccess, enabling subscribers to organize data based in hashed executor scope related to execution context.
+	struct FExecutorScopeKey
 	{
 		uint32 ExecutorId = FSubsonicExecutor::GetInvalidId();
 		Audio::FDeviceId DeviceId = INDEX_NONE;
-		ScopeHandleType ScopeHandle;
 
-		TScopeKey() = default;
-		TScopeKey(const FSubsonicExecutor& InExecutor, ScopeHandleType InScopeHandle)
+#if !NO_LOGGING
+		FCollectionHandle CollectionHandle;
+#endif // !NO_LOGGING
+
+		FExecutorScopeKey() = default;
+		FExecutorScopeKey(const FSubsonicExecutor& InExecutor)
 			: ExecutorId(InExecutor.GetId())
 			, DeviceId(InExecutor.GetDeviceId())
-			, ScopeHandle(MoveTemp(InScopeHandle))
+#if !NO_LOGGING
+			, CollectionHandle(InExecutor.GetCollectionHandle())
+#endif // !NO_LOGGING
 		{
 		}
 
-		friend uint32 GetTypeHash(const TScopeKey<ScopeHandleType>& Key)
+		friend uint32 GetTypeHash(const FExecutorScopeKey& Key)
 		{
-			const uint32 Hash1 = HashCombineFast(Key.ExecutorId, GetTypeHash(Key.DeviceId));
-			return HashCombineFast(Hash1, GetTypeHash(Key.ScopeHandle));
+			return HashCombineFast(Key.ExecutorId, GetTypeHash(Key.DeviceId));
 		}
 
-		friend bool operator==(const TScopeKey<ScopeHandleType>& A, const TScopeKey<ScopeHandleType>& B)
+		friend bool operator==(const FExecutorScopeKey& A, const FExecutorScopeKey& B)
 		{
-			return A.ExecutorId == B.ExecutorId && A.ScopeHandle == B.ScopeHandle && A.DeviceId == B.DeviceId;
+			return A.ExecutorId == B.ExecutorId && A.DeviceId == B.DeviceId;
 		}
 
 		FString ToString() const
 		{
-			return FString::Printf(TEXT("%s [ExecId: %u, AudioDevice Id: %u]"), *ScopeHandle.ToString(), ExecutorId, DeviceId);
+#if NO_LOGGING
+			return FString::Printf(TEXT("ExecId: %u, AudioDevice Id: %u"), ExecutorId, DeviceId);
+#else // !NO_LOGGING
+			return FString::Printf(TEXT("%s [ExecId: %u, AudioDevice Id: %u]"), *CollectionHandle.ToString(), ExecutorId, DeviceId);
+#endif // !NO_LOGGING
 		}
 	};
-
-	using FCollectionScopeKey = TScopeKey<FCollectionHandle>;
-	using FEventScopeKey = TScopeKey<FEventHandle>;
-	using FActionScopeKey = TScopeKey<FActionHandle>;
 } // namespace UE::Subsonic::Core
 
 #undef UE_API // SUBSONICCORE_API
diff --git a/Engine/Plugins/Experimental/Subsonic/Source/SubsonicCore/Public/SubsonicHandles.h b/Engine/Plugins/Experimental/Subsonic/Source/SubsonicCore/Public/SubsonicHandles.h
--- a/Engine/Plugins/Experimental/Subsonic/Source/SubsonicCore/Public/SubsonicHandles.h
+++ b/Engine/Plugins/Experimental/Subsonic/Source/SubsonicCore/Public/SubsonicHandles.h
@@ -18,7 +18,7 @@ namespace UE::Subsonic::Core
 	struct FCollectionHandle
 	{
 #if WITH_EDITORONLY_DATA
-		// Owning Collection definition's name (for logging/debugging)
+		// Owning Collection definition's name (for action initialization and logging/debugging)
 		FName CollectionName;
 #endif // WITH_EDITORONLY_DATA
 
diff --git a/Engine/Plugins/Experimental/Subsonic/Source/SubsonicCore/Public/SubsonicSubscriberDataStore.h b/Engine/Plugins/Experimental/Subsonic/Source/SubsonicCore/Public/SubsonicSubscriberDataStore.h
new file mode 100644
--- /dev/null
+++ b/Engine/Plugins/Experimental/Subsonic/Source/SubsonicCore/Public/SubsonicSubscriberDataStore.h
@@ -0,0 +1,183 @@
+// Copyright Epic Games, Inc. All Rights Reserved.
+#pragma once
+
+#include "Algo/ForEach.h"
+#include "Containers/Map.h"
+#include "SubsonicExecutor.h"
+#include "Templates/Function.h"
+
+
+namespace UE::Subsonic::Core
+{
+	// Generalized scope- and executor-keyed data store. Stores a DataType instance per scope level
+	// (Global, Executor) and provides convenience accessors, executor tracking, and iteration.
+	template <typename DataType>
+	class TSubscriberDataStore
+	{
+	private:
+		using FDataMap = TMap<FName, DataType>;
+		using FNameDataPair = TPair<FName, DataType>;
+		using FKeyDataMapPair = TPair<FExecutorScopeKey, FDataMap>;
+
+		FDataMap GlobalStore;
+		TMap<FExecutorScopeKey, FDataMap> ExecutorDataStore;
+
+	public:
+		// Empties all scoped stores
+		void Empty()
+		{
+			GlobalStore.Empty();
+			ExecutorDataStore.Empty();
+		}
+
+		// Empties all scoped stores, running the provided release function
+		// on all data values prior to destruction.
+		void Empty(TFunctionRef<void(FName, DataType&)> ReleaseFunc)
+		{
+			ForEach(ReleaseFunc);
+
+			auto IterPairs = [this, &ReleaseFunc](FKeyDataMapPair& Pair)
+			{
+				ForEach(Pair.Key, ReleaseFunc);
+			};
+			Algo::ForEach(ExecutorDataStore, IterPairs);
+			Empty();
+		}
+
+		// Finds associated data with a given global identifier
+		DataType* Find(FName GlobalName)
+		{
+			return GlobalStore.Find(GlobalName);
+		}
+
+		// Finds data associated with a given global identifier
+		const DataType* Find(FName GlobalName) const
+		{
+			return GlobalStore.Find(GlobalName);
+		}
+
+		// Finds data associated with a given executor and name
+		DataType* Find(const FExecutorScopeKey& Key, FName Name)
+		{
+			if (FDataMap* Map = ExecutorDataStore.Find(Key))
+			{
+				return Map->Find(Name);
+			}
+
+			return nullptr;
+		}
+
+		// Finds data associated with a given executor and name
+		const DataType* Find(const FExecutorScopeKey& Key, FName Name) const
+		{
+			if (const FDataMap* Map = ExecutorDataStore.Find(Key))
+			{
+				return Map->Find(Name);
+			}
+
+			return nullptr;
+		}
+
+
+		// Finds or adds data associated with a given global identifier
+		DataType& FindOrAdd(FName Name)
+		{
+			return GlobalStore.FindOrAdd(Name);
+		}
+
+		// Finds or adds data associated with a given executor & name.
+		DataType& FindOrAdd(const FExecutorScopeKey& Key, FName Name)
+		{
+			FDataMap& Map = ExecutorDataStore.FindOrAdd(Key);
+			return Map.FindOrAdd(Name);
+		}
+
+		// Iterate on all data values associated with the given key with all names.
+		void ForEach(const FExecutorScopeKey& InKey, TFunctionRef<void(FName, DataType&)> InFunc)
+		{
+			if (FDataMap* Map = ExecutorDataStore.Find(InKey))
+			{
+				auto FuncOnValues = [&InFunc](FNameDataPair& Pair) { InFunc(Pair.Key, Pair.Value); };
+				Algo::ForEach(*Map, FuncOnValues);
+			}
+		}
+
+		// Iterate on all data values associated with the given key with all names.
+		void ForEach(const FExecutorScopeKey& InKey, TFunctionRef<void(FName, DataType&)> InFunc) const
+		{
+			if (const FDataMap* Map = ExecutorDataStore.Find(InKey))
+			{
+				auto FuncOnValues = [&InFunc](const FNameDataPair& Pair) { InFunc(Pair.Key, Pair.Value); };
+				Algo::ForEach(*Map, FuncOnValues);
+			}
+		}
+
+		// Iterate on all executor-scoped data maps with their associated keys.
+		void ForEach(TFunctionRef<void(const FExecutorScopeKey&, const FName&, DataType&)> InFunc)
+		{
+			auto FuncOnPairs = [&InFunc](FKeyDataMapPair& MapPair)
+			{
+				Algo::ForEach(MapPair.Value, [&](TPair<FName, DataType>& DataPair)
+				{
+					InFunc(MapPair.Key, DataPair.Key, DataPair.Value);
+				});
+			};
+			Algo::ForEach(ExecutorDataStore, FuncOnPairs);
+		}
+
+		// Iterate on all executor-scoped data with their associated keys.
+		void ForEach(TFunctionRef<void(const FExecutorScopeKey&, const FName&, const DataType&)> InFunc) const
+		{
+			auto FuncOnPairs = [&InFunc](const FKeyDataMapPair& MapPair)
+			{
+				Algo::ForEach(MapPair.Value, [&](const TPair<FName, DataType>& DataPair)
+				{
+					InFunc(MapPair.Key, DataPair.Key, DataPair.Value);
+				});
+			};
+			Algo::ForEach(ExecutorDataStore, FuncOnPairs);
+		}
+
+		// Iterate on all values at global scope with all names.
+		void ForEach(TFunctionRef<void(FName, DataType&)> InFunc)
+		{
+			auto FuncOnValues = [&InFunc](FNameDataPair& Pair) { InFunc(Pair.Key, Pair.Value); };
+			Algo::ForEach(GlobalStore, FuncOnValues);
+		}
+
+		// Iterate on all values at global scope with all names.
+		void ForEach(TFunctionRef<void(FName, const DataType&)> InFunc) const
+		{
+			auto FuncOnValues = [&InFunc](const FNameDataPair& Pair) { InFunc(Pair.Key, Pair.Value); };
+			Algo::ForEach(GlobalStore, FuncOnValues);
+		}
+
+		// Removes the named data at the global scope level.
+		int32 Remove(FName Name)
+		{
+			return GlobalStore.Remove(Name);
+		}
+
+		// Removes data associated with the given name at the executor scope level.
+		int32 Remove(const FExecutorScopeKey& Key, FName Name)
+		{
+			int32 NumRemoved = 0;
+			if (FDataMap* Data = ExecutorDataStore.Find(Key))
+			{
+				NumRemoved = Data->Remove(Name);
+				if (Data->IsEmpty())
+				{
+					ExecutorDataStore.Remove(Key);
+				}
+			}
+
+			return NumRemoved;
+		}
+
+		// Removes all named data at the executor scope level.
+		int32 Remove(const FExecutorScopeKey& Key)
+		{
+			return ExecutorDataStore.Remove(Key);
+		}
+	};
+} // namespace UE::Subsonic::Core
diff --git a/Engine/Plugins/Experimental/Subsonic/Source/SubsonicEditor/Private/SubsonicEventCollectionEditor.cpp b/Engine/Plugins/Experimental/Subsonic/Source/SubsonicEditor/Private/SubsonicEventCollectionEditor.cpp
--- a/Engine/Plugins/Experimental/Subsonic/Source/SubsonicEditor/Private/SubsonicEventCollectionEditor.cpp
+++ b/Engine/Plugins/Experimental/Subsonic/Source/SubsonicEditor/Private/SubsonicEventCollectionEditor.cpp
@@ -324,12 +324,13 @@ namespace UE::Subsonic::Editor
 							const FName NewSelectionName = FName(*NewSelection);
 							if (USubsonicEditorSubsystem* EditorSubsystem = GEditor->GetEditorSubsystem<USubsonicEditorSubsystem>())
 							{
-								EditorSubsystem->ForEachActionStruct([this, &Action, NewSelectionName](const UScriptStruct& ActionStruct)
+								EditorSubsystem->ForEachActionStruct([this, &Action, NewSelectionName, &HandleRef](const UScriptStruct& ActionStruct)
 								{
 									if (ActionStruct.GetFName() == NewSelectionName)
 									{
 										TInstancedStruct<FSubsonicEventActionBase> NewConfig;
 										NewConfig.InitializeAsScriptStruct(&ActionStruct);
+										FSubsonicEventActionBase::InitializeDefaultActionName(NewConfig, HandleRef->GetEvent());
 										Action->SetAction(MoveTemp(NewConfig));
 										RebuildEventTree(true /* bRebuildListView */);
 									}
@@ -1548,8 +1549,7 @@ namespace UE::Subsonic::Editor
 			return;
 		}
 
-		const FSubsonicExecutor& Executor = PreviewExecutor->GetExecutor();
-		const FActionScopeKey ScopeKey(Executor, ActionHandle);
+		const FExecutorScopeKey ScopeKey(PreviewExecutor->GetExecutor());
 		USubsonicGeneratorSourceSubscriber* Subscriber = FindDeviceSubsystem<USubsonicGeneratorSourceSubscriber>(ScopeKey);
 		if (!Subscriber)
 		{
@@ -1562,16 +1562,8 @@ namespace UE::Subsonic::Editor
 			Subscriber->SetAuthoredParameters(PlayAction->Name, PlayAction->Parameters);
 			break;
 
-		case ESubsonicExecutionScope::Collection:
-			Subscriber->SetAuthoredParameters(ActionHandle.Event.Collection, PlayAction->Name, PlayAction->Parameters);
-			break;
-
-		case ESubsonicExecutionScope::Event:
-			Subscriber->SetAuthoredParameters(ActionHandle.Event, PlayAction->Name, PlayAction->Parameters);
-			break;
-
-		case ESubsonicExecutionScope::Action:
-			Subscriber->SetAuthoredParameters(ActionHandle, PlayAction->Name, PlayAction->Parameters);
+		case ESubsonicExecutionScope::Executor:
+			Subscriber->SetAuthoredParameters(ScopeKey, PlayAction->Name, PlayAction->Parameters);
 			break;
 
 		default:
diff --git a/Engine/Plugins/Experimental/Subsonic/Source/SubsonicEngine/Private/StandardActions/SubsonicAction_AudioComponent.cpp b/Engine/Plugins/Experimental/Subsonic/Source/SubsonicEngine/Private/StandardActions/SubsonicAction_AudioComponent.cpp
--- a/Engine/Plugins/Experimental/Subsonic/Source/SubsonicEngine/Private/StandardActions/SubsonicAction_AudioComponent.cpp
+++ b/Engine/Plugins/Experimental/Subsonic/Source/SubsonicEngine/Private/StandardActions/SubsonicAction_AudioComponent.cpp
@@ -2,10 +2,10 @@
 
 #include "StandardActions/SubsonicAction_AudioComponent.h"
 
+#include "Components/AudioComponent.h"
 #include "IAudioModulation.h"
 #include "Internationalization/Internationalization.h"
-#include "Sound/SoundBase.h"
-#include "Sound/SoundConcurrency.h"
+#include "StandardEventSubscribers/AudioComponentEventSubscriber.h"
 #include "SubsonicDeviceUtils.h"
 #include "SubsonicExecutor.h"
 #include "SubsonicHandles.h"
@@ -17,7 +17,7 @@ namespace UE::Subsonic
 {
 	namespace AudioComponentPrivate
 	{
-		UAudioComponent* AddScopedComponent(const Core::FActionScopeKey& InKey, ESubsonicExecutionScope Scope, FName Name)
+		UAudioComponent* AddScopedComponent(const Core::FExecutorScopeKey& InKey, ESubsonicExecutionScope Scope, FName Name)
 		{
 			UAudioComponent* Component = nullptr;
 			if (USubsonicAudioComponentSubscriber* AudioComponentSubscriber = FindDeviceSubsystem<USubsonicAudioComponentSubscriber>(InKey))
@@ -30,21 +30,9 @@ namespace UE::Subsonic
 					}
 					break;
 
-					case ESubsonicExecutionScope::Collection:
+					case ESubsonicExecutionScope::Executor:
 					{
-						Component = AudioComponentSubscriber->AddComponent(InKey.ScopeHandle.Event.Collection, Name);
-					}
-					break;
-
-					case ESubsonicExecutionScope::Event:
-					{
-						Component = AudioComponentSubscriber->AddComponent(InKey.ScopeHandle.Event, Name);
-					}
-					break;
-
-					case ESubsonicExecutionScope::Action:
-					{
-						Component = AudioComponentSubscriber->AddComponent(InKey.ScopeHandle, Name);
+						Component = AudioComponentSubscriber->AddComponent(InKey, Name);
 					}
 					break;
 				}
@@ -53,34 +41,22 @@ namespace UE::Subsonic
 			return Component;
 		}
 
-		UAudioComponent* FindScopedComponent(const Core::FActionScopeKey& InKey, ESubsonicExecutionScope Scope, FName Name)
+		UAudioComponent* FindScopedComponent(const Core::FExecutorScopeKey& InKey, ESubsonicExecutionScope Scope, FName Name)
 		{
 			UAudioComponent* Component = nullptr;
 			if (USubsonicAudioComponentSubscriber* AudioComponentSubscriber = FindDeviceSubsystem<USubsonicAudioComponentSubscriber>(InKey))
 			{
 				switch (Scope)
-					{
+				{
 					case ESubsonicExecutionScope::Global:
 					{
 						Component = AudioComponentSubscriber->FindComponent(Name);
 					}
 					break;
 
-					case ESubsonicExecutionScope::Collection:
+					case ESubsonicExecutionScope::Executor:
 					{
-						Component = AudioComponentSubscriber->FindComponent(InKey.ScopeHandle.Event.Collection, Name);
-					}
-					break;
-
-					case ESubsonicExecutionScope::Event:
-					{
-						Component = AudioComponentSubscriber->FindComponent(InKey.ScopeHandle.Event, Name);
-					}
-					break;
-
-					case ESubsonicExecutionScope::Action:
-					{
-						Component = AudioComponentSubscriber->FindComponent(InKey.ScopeHandle, Name);
+						Component = AudioComponentSubscriber->FindComponent(InKey, Name);
 					}
 					break;
 				}
@@ -89,7 +65,7 @@ namespace UE::Subsonic
 			return Component;
 		}
 
-		UAudioComponent* AccessComponent(const Core::FActionScopeKey& InKey, ESubsonicAudioComponentAccess Access, ESubsonicExecutionScope Scope, FName Name)
+		UAudioComponent* AccessComponent(const Core::FExecutorScopeKey& InKey, ESubsonicAudioComponentAccess Access, ESubsonicExecutionScope Scope, FName Name)
 		{
 			UAudioComponent* Component = nullptr;
 			if (Access == ESubsonicAudioComponentAccess::Find || Access == ESubsonicAudioComponentAccess::FindOrAdd)
@@ -119,14 +95,14 @@ namespace UE::Subsonic
 
 	void FSubsonicEventAction_AudioComponentModify::Execute(const Core::FSubsonicExecutor& InExecutor, const Core::FActionHandle& InHandle) const
 	{
-		const Core::FActionScopeKey ScopeKey(InExecutor, InHandle);
+		const Core::FExecutorScopeKey ScopeKey(InExecutor);
 		if (UAudioComponent* AudioComponent = AudioComponentPrivate::AccessComponent(ScopeKey, Access, Scope, Name))
 		{
 			for (const TInstancedStruct<FSubsonicEventAction_AudioComponentModifierBase>& Modifier : Modifiers)
 			{
 				if (const FSubsonicEventAction_AudioComponentModifierBase* ModPtr = Modifier.GetPtr())
 				{
-					ModPtr->Execute(*AudioComponent);
+					ModPtr->Execute(InExecutor, *AudioComponent);
 				}
 			}
 		}
@@ -165,7 +141,7 @@ namespace UE::Subsonic
 
 		return NameText;
 	}
-	
+
 	FText FSubsonicEventAction_AudioComponentPlay::GetDisplayInfo() const
 	{
 		if (!Name.IsNone())
@@ -184,7 +160,7 @@ namespace UE::Subsonic
 
 	void FSubsonicEventAction_AudioComponentPlay::Execute(const Core::FSubsonicExecutor& InExecutor, const Core::FActionHandle& InHandle) const
 	{
-		const Core::FActionScopeKey ScopeKey(InExecutor, InHandle);
+		const Core::FExecutorScopeKey ScopeKey(InExecutor);
 		if (UAudioComponent* AudioComponent = AudioComponentPrivate::AccessComponent(ScopeKey, Access, Scope, Name))
 		{
 			AudioComponent->SetSound(Sound);
@@ -194,7 +170,7 @@ namespace UE::Subsonic
 
 	void FSubsonicEventAction_AudioComponentStop::Execute(const Core::FSubsonicExecutor& InExecutor, const Core::FActionHandle& InHandle) const
 	{
-		const Core::FActionScopeKey ScopeKey(InExecutor, InHandle);
+		const Core::FExecutorScopeKey ScopeKey(InExecutor);
 		if (UAudioComponent* AudioComponent = AudioComponentPrivate::FindScopedComponent(ScopeKey, Scope, Name))
 		{
 			AudioComponent->Stop();
@@ -213,34 +189,46 @@ namespace UE::Subsonic
 	}
 #endif // WITH_EDITOR
 
-	void FSubsonicEventAction_AudioComponentModifier_Play::Execute(UAudioComponent& AudioComponent) const
+	void FSubsonicEventAction_AudioComponentModifier_ExecuteOnFinished::Execute(const Core::FSubsonicExecutor& Executor, UAudioComponent& AudioComponent) const
+	{
+		AudioComponent.OnAudioFinishedNative.Remove(ActiveDelegate);
+		ActiveDelegate = AudioComponent.OnAudioFinishedNative.AddLambda([ExecWeakPtr = Executor.AsWeak(), EventName = Event.GetTagName()](UAudioComponent* Component)
+		{
+			if (TSharedPtr<const Core::FSubsonicExecutor> ExecPtr = ExecWeakPtr.Pin())
+			{
+				ExecPtr->ExecuteEvent(EventName);
+			}
+		});
+	}
+
+	void FSubsonicEventAction_AudioComponentModifier_Play::Execute(const Core::FSubsonicExecutor& InExecutor, UAudioComponent& AudioComponent) const
 	{
 		AudioComponent.Play(StartTime);
 	}
 
-	void FSubsonicEventAction_AudioComponentModifier_SetAttenuation::Execute(UAudioComponent& AudioComponent) const
+	void FSubsonicEventAction_AudioComponentModifier_SetAttenuation::Execute(const Core::FSubsonicExecutor& InExecutor, UAudioComponent& AudioComponent) const
 	{
 		AudioComponent.SetAttenuationSettings(Attenuation);
 	}
 
-	void FSubsonicEventAction_AudioComponentModifier_SetConcurrency::Execute(UAudioComponent& AudioComponent) const
+	void FSubsonicEventAction_AudioComponentModifier_SetConcurrency::Execute(const Core::FSubsonicExecutor& InExecutor, UAudioComponent& AudioComponent) const
 	{
 		AudioComponent.ConcurrencySet = Concurrency;
 	}
 
-	void FSubsonicEventAction_AudioComponentModifier_SetModulationRouting::Execute(UAudioComponent& AudioComponent) const
+	void FSubsonicEventAction_AudioComponentModifier_SetModulationRouting::Execute(const Core::FSubsonicExecutor& InExecutor, UAudioComponent& AudioComponent) const
 	{
 		TSet<USoundModulatorBase*> ModulatorSet;
 		Algo::Transform(Modulators, ModulatorSet, [](const TObjectPtr<USoundModulatorBase>& Modulator) { return Modulator.Get(); });
 		AudioComponent.SetModulationRouting(ModulatorSet, Destination, RoutingMethod);
 	}
 
-	void FSubsonicEventAction_AudioComponentModifier_SetSound::Execute(UAudioComponent& AudioComponent) const
+	void FSubsonicEventAction_AudioComponentModifier_SetSound::Execute(const Core::FSubsonicExecutor& InExecutor, UAudioComponent& AudioComponent) const
 	{
 		AudioComponent.SetSound(Sound);
 	}
 
-	void FSubsonicEventAction_AudioComponentModifier_Stop::Execute(UAudioComponent& AudioComponent) const
+	void FSubsonicEventAction_AudioComponentModifier_Stop::Execute(const Core::FSubsonicExecutor& InExecutor, UAudioComponent& AudioComponent) const
 	{
 		AudioComponent.Stop();
 	}
diff --git a/Engine/Plugins/Experimental/Subsonic/Source/SubsonicEngine/Private/StandardActions/SubsonicAction_AudioComponent.h b/Engine/Plugins/Experimental/Subsonic/Source/SubsonicEngine/Private/StandardActions/SubsonicAction_AudioComponent.h
--- a/Engine/Plugins/Experimental/Subsonic/Source/SubsonicEngine/Private/StandardActions/SubsonicAction_AudioComponent.h
+++ b/Engine/Plugins/Experimental/Subsonic/Source/SubsonicEngine/Private/StandardActions/SubsonicAction_AudioComponent.h
@@ -1,8 +1,9 @@
 // Copyright Epic Games, Inc. All Rights Reserved.
 #pragma once
 
+#include "GameplayTagContainer.h"
+#include "Sound/SoundBase.h"
 #include "Sound/SoundModulationDestination.h"
-#include "StandardEventSubscribers/AudioComponentEventSubscriber.h"
 #include "StructUtils/InstancedStruct.h"
 #include "SubsonicEventCollection.h"
 #include "SubsonicExecutor.h"
@@ -13,7 +14,6 @@
 // Forward Declarations
 class UAudioComponent;
 class USoundConcurrency;
-class USoundModulatorBase;
 
 struct FSubsonicEventAction_AudioComponentModifierBase;
 
@@ -50,13 +50,14 @@ namespace UE::Subsonic
 
 		// Addressable name of audio component instance.
 		UPROPERTY(EditAnywhere, Category = "Subsonic")
-		FName Name = "Component01";
+		FName Name;
 
-		// Scope of component accessor.  If set to event, provided name is used
-		// to address the same component within the event, collection the same component
-		// within the collection, etc.
+		// Scope of component accessor.  If set to executor, named AudioComponent is
+		// accessible only by pool of sibling event actions for a given
+		// executor instance. If set to global, accesses any named AudioComponent
+		// in the global pool.
 		UPROPERTY(EditAnywhere, Category = "Subsonic")
-		ESubsonicExecutionScope Scope = ESubsonicExecutionScope::Collection;
+		ESubsonicExecutionScope Scope = ESubsonicExecutionScope::Executor;
 
 		UPROPERTY(EditAnywhere, Category = "Subsonic")
 		ESubsonicAudioComponentAccess Access = ESubsonicAudioComponentAccess::FindOrAdd;
@@ -81,17 +82,18 @@ namespace UE::Subsonic
 	public:
 		// Addressable name of audio component instance.
 		UPROPERTY(EditAnywhere, Category = "Subsonic")
-		FName Name = "Component01";
+		FName Name;
 
 		// Sound to play.
 		UPROPERTY(EditAnywhere, Category = "Subsonic")
 		TObjectPtr<USoundBase> Sound;
 
-		// Scope of component accessor.  If set to event, provided name is used
-		// to address the same component within the event, collection the same component
-		// within the collection, etc.
+		// Scope of component accessor.  If set to executor, named AudioComponent is
+		// accessible only by pool of sibling event actions for a given
+		// executor instance. If set to global, accesses any named AudioComponent
+		// in the global pool.
 		UPROPERTY(EditAnywhere, Category = "Subsonic")
-		ESubsonicExecutionScope Scope = ESubsonicExecutionScope::Collection;
+		ESubsonicExecutionScope Scope = ESubsonicExecutionScope::Executor;
 
 		UPROPERTY(EditAnywhere, Category = "Subsonic")
 		ESubsonicAudioComponentAccess Access = ESubsonicAudioComponentAccess::FindOrAdd;
@@ -112,13 +114,14 @@ namespace UE::Subsonic
 
 		// Addressable name of component/sound instance.
 		UPROPERTY(EditAnywhere, Category = "Subsonic")
-		FName Name = "Component01";
+		FName Name;
 
-		// Scope of component to play.  If set to event, provided name is used
-		// to address the existing component within the event, collection the existing component
-		// within the collection, etc.
+		// Scope of component accessor.  If set to executor, named AudioComponent is
+		// accessible only by pool of sibling event actions for a given
+		// executor instance. If set to global, accesses any named AudioComponent
+		// in the global pool.
 		UPROPERTY(EditAnywhere, Category = "Subsonic")
-		ESubsonicExecutionScope Scope = ESubsonicExecutionScope::Collection;
+		ESubsonicExecutionScope Scope = ESubsonicExecutionScope::Executor;
 	};
 
 	USTRUCT(MinimalAPI, meta = (Abstract, Hidden))
@@ -129,7 +132,27 @@ namespace UE::Subsonic
 	public:
 		virtual ~FSubsonicEventAction_AudioComponentModifierBase() = default;
 
-		virtual void Execute(UAudioComponent& AudioComponent) const { }
+		virtual void Execute(const Core::FSubsonicExecutor& InExecutor, UAudioComponent& AudioComponent) const { }
+	};
+
+	// Executes event with the given name within the parent collection on completion
+	USTRUCT(MinimalAPI, DisplayName = "Execute On Finished")
+	struct FSubsonicEventAction_AudioComponentModifier_ExecuteOnFinished : public FSubsonicEventAction_AudioComponentModifierBase
+	{
+		GENERATED_BODY()
+
+	public:
+		virtual ~FSubsonicEventAction_AudioComponentModifier_ExecuteOnFinished() = default;
+
+		virtual void Execute(const Core::FSubsonicExecutor& InExecutor, UAudioComponent& AudioComponent) const override;
+
+		// Event in parent collection to execute.
+		UPROPERTY(EditAnywhere, Category = "Subsonic")
+		FGameplayTag Event;
+
+	private:
+		// Should really be tracked by subscriber to support cancelation, replacement, data-oriented execution ordering, const correctness, etc.
+		mutable FDelegateHandle ActiveDelegate;
 	};
 
 	USTRUCT(MinimalAPI, DisplayName = "Play")
@@ -140,7 +163,7 @@ namespace UE::Subsonic
 	public:
 		virtual ~FSubsonicEventAction_AudioComponentModifier_Play() = default;
 
-		virtual void Execute(UAudioComponent& AudioComponent) const override;
+		virtual void Execute(const Core::FSubsonicExecutor& InExecutor, UAudioComponent& AudioComponent) const override;
 
 		// Play StartTime Offset (plays immediately if 0.0)
 		UPROPERTY(EditAnywhere, Category = "Subsonic")
@@ -155,7 +178,7 @@ namespace UE::Subsonic
 	public:
 		virtual ~FSubsonicEventAction_AudioComponentModifier_SetAttenuation() = default;
 
-		virtual void Execute(UAudioComponent& AudioComponent) const override;
+		virtual void Execute(const Core::FSubsonicExecutor& InExecutor, UAudioComponent& AudioComponent) const override;
 
 		UPROPERTY(EditAnywhere, Category = "Subsonic")
 		TObjectPtr<USoundAttenuation> Attenuation;
@@ -169,7 +192,7 @@ namespace UE::Subsonic
 	public:
 		virtual ~FSubsonicEventAction_AudioComponentModifier_SetConcurrency() = default;
 
-		virtual void Execute(UAudioComponent& AudioComponent) const override;
+		virtual void Execute(const Core::FSubsonicExecutor& InExecutor, UAudioComponent& AudioComponent) const override;
 
 		UPROPERTY(EditAnywhere, Category = "Subsonic")
 		TSet<TObjectPtr<USoundConcurrency>> Concurrency;
@@ -183,14 +206,14 @@ namespace UE::Subsonic
 	public:
 		virtual ~FSubsonicEventAction_AudioComponentModifier_SetModulationRouting() = default;
 
-		virtual void Execute(UAudioComponent& AudioComponent) const override;
+		virtual void Execute(const Core::FSubsonicExecutor& InExecutor, UAudioComponent& AudioComponent) const override;
 
 		UPROPERTY(EditAnywhere, Category = "Subsonic")
 		TSet<TObjectPtr<USoundModulatorBase>> Modulators;
 
 		UPROPERTY(EditAnywhere, Category = "Subsonic")
 		EModulationDestination Destination = EModulationDestination::Volume;
-	
+
 		UPROPERTY(EditAnywhere, Category = "Subsonic")
 		EModulationRouting RoutingMethod = EModulationRouting::Inherit;
 	};
@@ -203,7 +226,7 @@ namespace UE::Subsonic
 	public:
 		virtual ~FSubsonicEventAction_AudioComponentModifier_SetSound() = default;
 
-		virtual void Execute(UAudioComponent& AudioComponent) const override;
+		virtual void Execute(const Core::FSubsonicExecutor& InExecutor, UAudioComponent& AudioComponent) const override;
 
 		// Sound to play.
 		UPROPERTY(EditAnywhere, Category = "Subsonic")
@@ -218,6 +241,6 @@ namespace UE::Subsonic
 	public:
 		virtual ~FSubsonicEventAction_AudioComponentModifier_Stop() = default;
 
-		virtual void Execute(UAudioComponent& AudioComponent) const override;
+		virtual void Execute(const Core::FSubsonicExecutor& InExecutor, UAudioComponent& AudioComponent) const override;
 	};
-} // namespace UE::Subsonic
\ No newline at end of file
+} // namespace UE::Subsonic
diff --git a/Engine/Plugins/Experimental/Subsonic/Source/SubsonicEngine/Private/StandardActions/SubsonicAction_GeneratorSource.cpp b/Engine/Plugins/Experimental/Subsonic/Source/SubsonicEngine/Private/StandardActions/SubsonicAction_GeneratorSource.cpp
--- a/Engine/Plugins/Experimental/Subsonic/Source/SubsonicEngine/Private/StandardActions/SubsonicAction_GeneratorSource.cpp
+++ b/Engine/Plugins/Experimental/Subsonic/Source/SubsonicEngine/Private/StandardActions/SubsonicAction_GeneratorSource.cpp
@@ -24,7 +24,7 @@ namespace UE::Subsonic
 		FSubsonicParameterStore MergedParams = Parameters;
 		MergedParams.MergeFrom(InExecutor.GetParameters());
 
-		const FActionScopeKey ScopeKey(InExecutor, InHandle);
+		const FExecutorScopeKey ScopeKey(InExecutor);
 		if (USubsonicGeneratorSourceSubscriber* GeneratorSourceSubscriber = FindDeviceSubsystem<USubsonicGeneratorSourceSubscriber>(ScopeKey))
 		{
 			switch (Scope)
@@ -33,16 +33,8 @@ namespace UE::Subsonic
 				GeneratorSourceSubscriber->PlaySound(Name, *Sound, MergedParams);
 				break;
 
-			case ESubsonicExecutionScope::Collection:
-				GeneratorSourceSubscriber->PlaySound(InHandle.Event.Collection, Name, *Sound, MergedParams);
-				break;
-
-			case ESubsonicExecutionScope::Event:
-				GeneratorSourceSubscriber->PlaySound(InHandle.Event, Name, *Sound, MergedParams);
-				break;
-
-			case ESubsonicExecutionScope::Action:
-				GeneratorSourceSubscriber->PlaySound(InHandle, Name, *Sound, MergedParams);
+			case ESubsonicExecutionScope::Executor:
+				GeneratorSourceSubscriber->PlaySound(ScopeKey, Name, *Sound, MergedParams);
 				break;
 
 			default:
@@ -51,11 +43,23 @@ namespace UE::Subsonic
 		}
 	}
 
+#if WITH_EDITOR
+	FText FSubsonicEventAction_GeneratorSourcePlay::GetDisplayInfo() const
+	{
+		if (!Name.IsNone())
+		{
+			return FText::FromName(Name);
+		}
+
+		return Super::GetDisplayInfo();
+	}
+#endif // WITH_EDITOR
+
 	void FSubsonicEventAction_GeneratorSourceStop::Execute(const Core::FSubsonicExecutor& InExecutor, const Core::FActionHandle& InHandle) const
 	{
 		using namespace UE::Subsonic::Core;
 
-		const FActionScopeKey ScopeKey(InExecutor, InHandle);
+		const FExecutorScopeKey ScopeKey(InExecutor);
 		if (USubsonicGeneratorSourceSubscriber* GeneratorSourceSubscriber = FindDeviceSubsystem<USubsonicGeneratorSourceSubscriber>(ScopeKey))
 		{
 			switch (Scope)
@@ -64,16 +68,8 @@ namespace UE::Subsonic
 				GeneratorSourceSubscriber->StopSound(Name);
 				break;
 
-			case ESubsonicExecutionScope::Collection:
-				GeneratorSourceSubscriber->StopSound(InHandle.Event.Collection, Name);
-				break;
-
-			case ESubsonicExecutionScope::Event:
-				GeneratorSourceSubscriber->StopSound(InHandle.Event, Name);
-				break;
-
-			case ESubsonicExecutionScope::Action:
-				GeneratorSourceSubscriber->StopSound(InHandle, Name);
+			case ESubsonicExecutionScope::Executor:
+				GeneratorSourceSubscriber->StopSound(ScopeKey, Name);
 				break;
 
 			default:
@@ -81,4 +77,16 @@ namespace UE::Subsonic
 			}
 		}
 	}
-} // namespace UE::Subsonic
\ No newline at end of file
+
+#if WITH_EDITOR
+	FText FSubsonicEventAction_GeneratorSourceStop::GetDisplayInfo() const
+	{
+		if (!Name.IsNone())
+		{
+			return FText::FromName(Name);
+		}
+
+		return Super::GetDisplayInfo();
+	}
+#endif // WITH_EDITOR
+} // namespace UE::Subsonic
diff --git a/Engine/Plugins/Experimental/Subsonic/Source/SubsonicEngine/Private/StandardEventSubscribers/AudioComponentEventSubscriber.cpp b/Engine/Plugins/Experimental/Subsonic/Source/SubsonicEngine/Private/StandardEventSubscribers/AudioComponentEventSubscriber.cpp
--- a/Engine/Plugins/Experimental/Subsonic/Source/SubsonicEngine/Private/StandardEventSubscribers/AudioComponentEventSubscriber.cpp
+++ b/Engine/Plugins/Experimental/Subsonic/Source/SubsonicEngine/Private/StandardEventSubscribers/AudioComponentEventSubscriber.cpp
@@ -2,18 +2,13 @@
 
 #include "StandardEventSubscribers/AudioComponentEventSubscriber.h"
 
-#include "Algo/Transform.h"
-#include "AudioDevice.h"
-#include "AudioDeviceManager.h"
-#include "Engine/Engine.h"
+#include "Sound/SoundBase.h"
 #include "SubsonicCoreLog.h"
 #include "SubsonicExecutor.h"
-#include "SubsonicHandles.h"
-#include "SubsonicSubsystem.h"
+#include "SubsonicSubscriberDataStore.h"
 #include "UObject/GCObject.h"
 #include "UObject/UObjectGlobals.h"
 
-#define SUBSONIC_ASSERT_ON_STALE_COLLECTION_EXECUTORS (1 && DO_CHECK)
 
 namespace UE::Subsonic
 {
@@ -21,12 +16,32 @@ namespace UE::Subsonic
 	{
 		struct FInstanceAudioComponents : public FGCObject
 		{
+		private:
 			// Currently addressable component via handle
 			TObjectPtr<UAudioComponent> ActiveComponent;
+			FDelegateHandle OnFinishHandle;
 
 			// Components no longer addressable from actions but actively
 			// playing. Continually tracked here to avoid GC while finishing.
 			TArray<TObjectPtr<UAudioComponent>> ReleasedComponents;
+			TArray<FDelegateHandle> OnReleaseOnFinishHandles;
+
+		public:
+			const TObjectPtr<UAudioComponent> GetActiveComponent() const
+			{
+				return ActiveComponent;
+			}
+
+			const TArray<TObjectPtr<UAudioComponent>> GetReleasedComponents() const
+			{
+				return ReleasedComponents;
+			}
+
+			void SetActiveComponent(TObjectPtr<UAudioComponent> Component, FDelegateHandle&& InOnFinishHandle)
+			{
+				ActiveComponent = Component;
+				OnFinishHandle = MoveTemp(InOnFinishHandle);
+			}
 
 			virtual void AddReferencedObjects(FReferenceCollector& Collector) override
 			{
@@ -39,6 +54,27 @@ namespace UE::Subsonic
 				return "UE::Subsonic::InstanceAudioComponents";
 			}
 
+			void OnComponentFinished(UAudioComponent* FinishingComponent)
+			{
+				if (ActiveComponent.Get() == FinishingComponent)
+				{
+					ActiveComponent.Get()->OnAudioFinishedNative.Remove(OnFinishHandle);
+					ActiveComponent = nullptr;
+					OnFinishHandle.Reset();
+				}
+
+				for (int32 Index = ReleasedComponents.Num() - 1; Index >= 0; --Index)
+				{
+					if (ReleasedComponents[Index].Get() == FinishingComponent)
+					{
+						ReleasedComponents[Index]->OnAudioFinishedNative.Remove(OnReleaseOnFinishHandles[Index]);
+						ReleasedComponents.RemoveAtSwap(Index);
+						OnReleaseOnFinishHandles.RemoveAtSwap(Index);
+					}
+				}
+
+			}
+
 			bool IsEmpty() const
 			{
 				return !ActiveComponent.Get() && ReleasedComponents.IsEmpty();
@@ -49,24 +85,35 @@ namespace UE::Subsonic
 				if (ActiveComponent.Get() && ActiveComponent->IsPlaying())
 				{
 					ReleasedComponents.Add(ActiveComponent);
+					OnReleaseOnFinishHandles.Add(OnFinishHandle);
 				}
+				OnFinishHandle.Reset();
 				ActiveComponent = nullptr;
 			}
 
-			void StopAll() const
+			void StopAll()
 			{
 				if (UAudioComponent* Component = ActiveComponent.Get())
 				{
 					Component->Stop();
+					Component->OnAudioFinishedNative.Remove(OnFinishHandle);
 				}
 
-				for (const TObjectPtr<UAudioComponent>& ReleasedComponent : ReleasedComponents)
+				for (int32 Index = 0; Index < ReleasedComponents.Num(); ++Index)
 				{
-					if (UAudioComponent* Component = ReleasedComponent.Get())
+					if (UAudioComponent* Component = ReleasedComponents[Index].Get())
 					{
 						Component->Stop();
+						Component->OnAudioFinishedNative.Remove(OnReleaseOnFinishHandles[Index]);
 					}
 				}
+
+				ActiveComponent = nullptr;
+				OnFinishHandle.Reset();
+
+				ReleasedComponents.Reset();
+				OnReleaseOnFinishHandles.Reset();
+
 			}
 		};
 
@@ -75,13 +122,8 @@ namespace UE::Subsonic
 
 		class FSubscriberImpl : public TSharedFromThis<FSubscriberImpl>
 		{
-			FAudioComponentsCache GlobalComponents;
-
-			TMap<Core::FCollectionHandle, FAudioComponentsCache> CollectionComponents;
-			TMap<Core::FEventHandle, FAudioComponentsCache> EventComponents;
-			TMap<Core::FActionHandle, FAudioComponentsCache> ActionComponents;
-
-			TMap<uint32, Core::FCollectionHandle> CollectionExecutors;
+		public:
+			Core::TSubscriberDataStore<FInstanceAudioComponentsPtr> ComponentsStore;
 
 			static UAudioComponent* CreateComponent(UObject& Parent, Audio::FDeviceId DeviceId, FName Name)
 			{
@@ -99,365 +141,193 @@ namespace UE::Subsonic
 				return NewComponent;
 			}
 
-			static UAudioComponent* AddCachedComponent(
-				UObject& Parent,
-				Audio::FDeviceId DeviceId,
-				const FInstanceAudioComponentsPtr& InstanceComponents,
-				FName Name,
-				bool bReleaseExisting,
-				bool& bOutStoppedComponent)
+			template <typename KeyType>
+			bool ContainsCacheComponent(const TMap<KeyType, FAudioComponentsCache>& ScopeComponentMap, const KeyType& InKey, FName Name) const
 			{
-				bOutStoppedComponent = false;
-
-				if (bReleaseExisting)
+				if (const FAudioComponentsCache* Cache = ScopeComponentMap.Find(InKey))
 				{
-					if (UAudioComponent* Component = InstanceComponents->ActiveComponent.Get(); Component && Component->IsPlaying())
+					if (const FInstanceAudioComponentsPtr* CacheComponentsPtr = Cache->Find(Name))
 					{
-						if (USoundBase* Sound = Component->GetSound(); Sound && !Sound->IsOneShot())
-						{
-							bOutStoppedComponent = true;
-							Component->Stop();
-						}
+						return CacheComponentsPtr->IsValid();
 					}
-
-					InstanceComponents->ReleaseComponent();
-				}
-
-				TObjectPtr<UAudioComponent>& ComponentRef = InstanceComponents->ActiveComponent;
-				if (!ComponentRef)
-				{
-					ComponentRef = CreateComponent(Parent, DeviceId, Name);
 				}
 
-				return ComponentRef.Get();
+				return false;
 			}
 
-			template <typename HandleType>
-			UAudioComponent* AddCachedComponent(
-				UObject& Parent,
-				Audio::FDeviceId DeviceId,
-				const HandleType& InHandle,
-				TMap<HandleType, FAudioComponentsCache>& OutScopeComponentMap,
-				FName Name,
-				bool bReleaseExisting)
+		public:
+			void Deinitialize()
 			{
-				UAudioComponent* CacheComponent = nullptr;
+				auto StopAll = [](FName, FInstanceAudioComponentsPtr& InstanceComponents)
+				{
+					if (InstanceComponents.IsValid())
+					{
+						InstanceComponents->StopAll();
+					}
+				};
+				ComponentsStore.Empty(StopAll);
+			}
 
-				FAudioComponentsCache& Cache = OutScopeComponentMap.FindOrAdd(InHandle);
-				FInstanceAudioComponentsPtr& InstanceComponents = Cache.FindOrAdd(Name);
+			UAudioComponent* AddComponent(UObject& Parent, Audio::FDeviceId DeviceId, FName Name, bool bReleaseExisting)
+			{
+				FInstanceAudioComponentsPtr& InstanceComponents = ComponentsStore.FindOrAdd(Name);
 				if (!InstanceComponents.IsValid())
 				{
 					InstanceComponents = MakeUnique<FInstanceAudioComponents>();
 				}
 
-				bool bStoppedComponent = false;
-				CacheComponent = AddCachedComponent(Parent, DeviceId, InstanceComponents, Name, bReleaseExisting, bStoppedComponent);
-				if (bStoppedComponent)
-				{
-					UE_LOG(LogSubsonic, Warning,
-						TEXT("Stopping released sound with handle '%s': "
-							"Sound is not one shot and continued play would result in leaked playback."),
-						*InHandle.ToString())
-				}
-				auto RemoveComponentLambda = [FinishHandle = InHandle, &OutScopeComponentMap, Name](UAudioComponent* FinishingComponent)
+				if (bReleaseExisting)
 				{
-					if (FAudioComponentsCache* Cache = OutScopeComponentMap.Find(FinishHandle))
+					UAudioComponent* Component = InstanceComponents->GetActiveComponent();
+					if (Component && Component->IsPlaying())
 					{
-						if (FInstanceAudioComponentsPtr* CacheComponentsPtr = Cache->Find(Name); CacheComponentsPtr && CacheComponentsPtr->IsValid())
-						{
-							FInstanceAudioComponentsPtr& CacheComponents = *CacheComponentsPtr;
-							if (CacheComponents->ActiveComponent.Get() == FinishingComponent)
-							{
-								CacheComponents->ActiveComponent = nullptr;
-							}
-
-							CacheComponents->ReleasedComponents.RemoveAllSwap([&FinishingComponent](const TObjectPtr<UAudioComponent>& Released)
-							{
-								return Released.Get() == FinishingComponent;
-							});
-
-							if (CacheComponents->IsEmpty())
-							{
-								Cache->Remove(Name);
-							}
-						}
-						if (Cache->IsEmpty())
+						if (USoundBase* Sound = Component->GetSound(); Sound && !Sound->IsOneShot())
 						{
-							OutScopeComponentMap.Remove(FinishHandle);
+							UE_LOG(LogSubsonic, Warning,
+								TEXT("Stopping released global sound with name '%s': "
+									"Sound is not one shot and continued play would result in leaked playback."),
+								*Name.ToString())
+							Component->Stop();
 						}
 					}
-				};
 
-				CacheComponent->OnAudioFinishedNative.AddSPLambda(this, RemoveComponentLambda);
-				return CacheComponent;
-			}
-
-			template <typename HandleType>
-			UAudioComponent* FindCacheComponent(const TMap<HandleType, FAudioComponentsCache>& ScopeComponentMap, const HandleType& InHandle, FName Name) const
-			{
-				if (const FAudioComponentsCache* Cache = ScopeComponentMap.Find(InHandle))
-				{
-					if (const FInstanceAudioComponentsPtr* CacheComponentsPtr = Cache->Find(Name); CacheComponentsPtr && CacheComponentsPtr->IsValid())
-					{
-						return (*CacheComponentsPtr)->ActiveComponent.Get();
-					}
+					InstanceComponents->ReleaseComponent();
 				}
 
-				return nullptr;
-			}
-
-			template <typename HandleType>
-			bool ContainsCacheComponent(const TMap<HandleType, FAudioComponentsCache>& ScopeComponentMap, const HandleType& InHandle, FName Name) const
-			{
-				if (const FAudioComponentsCache* Cache = ScopeComponentMap.Find(InHandle))
+				TObjectPtr<UAudioComponent> ActiveComponent = InstanceComponents->GetActiveComponent();
+				if (!ActiveComponent)
 				{
-					if (const FInstanceAudioComponentsPtr* CacheComponentsPtr = Cache->Find(Name))
+					ActiveComponent = CreateComponent(Parent, DeviceId, Name);
+					FDelegateHandle OnFinishedHandle = ActiveComponent->OnAudioFinishedNative.AddLambda([WeakImpl = AsWeak(), Name](UAudioComponent* FinishingComponent)
 					{
-						return CacheComponentsPtr->IsValid();
-					}
+						if (TSharedPtr<FSubscriberImpl> Impl = WeakImpl.Pin())
+						{
+							if (FInstanceAudioComponentsPtr* ComponentsPtr = Impl->ComponentsStore.Find(Name))
+							{
+								(*ComponentsPtr)->OnComponentFinished(FinishingComponent);
+								if ((*ComponentsPtr)->IsEmpty())
+								{
+									Impl->ComponentsStore.Remove(Name);
+								}
+							}
+						}
+					});
+					InstanceComponents->SetActiveComponent(ActiveComponent, MoveTemp(OnFinishedHandle));
 				}
 
-				return false;
+				return ActiveComponent;
 			}
 
-			template <typename HandleType>
-			void RemoveCacheComponent(TMap<HandleType, FAudioComponentsCache>& ScopeComponentMap, const HandleType& InHandle, FName Name)
+			UAudioComponent* AddComponent(UObject& Parent, Audio::FDeviceId DeviceId, const Core::FExecutorScopeKey& InKey, FName Name, bool bReleaseExisting)
 			{
-				if (FAudioComponentsCache* Cache = ScopeComponentMap.Find(InHandle))
+				FInstanceAudioComponentsPtr& InstanceComponents = ComponentsStore.FindOrAdd(InKey, Name);
+				if (InstanceComponents.IsValid())
 				{
-					FInstanceAudioComponentsPtr* InstanceComponents = Cache->Find(Name);
-					if (InstanceComponents && InstanceComponents->IsValid())
+					if (bReleaseExisting)
 					{
-						(*InstanceComponents)->StopAll();
-					}
+						UAudioComponent* Component = InstanceComponents->GetActiveComponent();
+						if (Component && Component->IsPlaying())
+						{
+							if (USoundBase* Sound = Component->GetSound(); Sound && !Sound->IsOneShot())
+							{
+								UE_LOG(LogSubsonic, Warning,
+									TEXT("Stopping released sound with handle '%s': "
+										"Sound is not one shot and continued play would result in leaked playback."),
+									*InKey.ToString())
+								Component->Stop();
+							}
+						}
 
-					Cache->Remove(Name);
-					if (Cache->IsEmpty())
-					{
-						ScopeComponentMap.Remove(InHandle);
+						InstanceComponents->ReleaseComponent();
 					}
-				}
-			}
-
-			template <typename THandleType>
-			void ResetScopeComponents(TMap<THandleType, FAudioComponentsCache>& ScopeComponentMap)
-			{
-				for (const TPair<THandleType, FAudioComponentsCache>& HandleCachePair : ScopeComponentMap)
-				{
-					for (const TPair<FName, FInstanceAudioComponentsPtr>& Pair : HandleCachePair.Value)
+					else
 					{
-						const FInstanceAudioComponentsPtr& InstanceComponents = Pair.Value;
-						check(InstanceComponents.IsValid());
-						InstanceComponents->StopAll();
-					}
-				}
-				ScopeComponentMap.Reset();
-			};
 
-		public:
-			void Deinitialize()
-			{
-				for (const TPair<FName, FInstanceAudioComponentsPtr>& Pair : GlobalComponents)
-				{
-					const FInstanceAudioComponentsPtr& InstanceComponents = Pair.Value;
-					check(InstanceComponents.IsValid());
-					InstanceComponents->StopAll();
+					}
 				}
-				GlobalComponents.Empty();
-
-				ResetScopeComponents(CollectionComponents);
-				ResetScopeComponents(EventComponents);
-				ResetScopeComponents(ActionComponents);
-			}
-
-			UAudioComponent* AddComponent(UObject& Parent, Audio::FDeviceId DeviceId, FName Name, bool bReleaseExisting)
-			{
-				FInstanceAudioComponentsPtr& InstanceComponents = GlobalComponents.FindOrAdd(Name);
-				if (!InstanceComponents.IsValid())
+				else
 				{
 					InstanceComponents = MakeUnique<FInstanceAudioComponents>();
 				}
 
-				bool bStoppedComponent = false;
-				if (UAudioComponent* CacheComponent = AddCachedComponent(Parent, DeviceId, InstanceComponents, Name, bReleaseExisting, bStoppedComponent))
+				TObjectPtr<UAudioComponent> ActiveComponent = InstanceComponents->GetActiveComponent();
+				if (!ActiveComponent)
 				{
-					if (bStoppedComponent)
-					{
-						UE_LOG(LogSubsonic, Warning,
-							TEXT("Stopping released global sound '%s': "
-								"Sound is not one shot and continued play would result in leaked playback."),
-							*Name.ToString())
-					}
-
-					auto RemoveComponentLambda = [this, Name](UAudioComponent* FinishingComponent)
+					ActiveComponent = CreateComponent(Parent, DeviceId, Name);
+					FDelegateHandle OnFinishedHandle = ActiveComponent->OnAudioFinishedNative.AddLambda([WeakImpl = AsWeak(), FinishKey = InKey, Name](UAudioComponent* FinishingComponent)
 					{
-						if (FInstanceAudioComponentsPtr* Components = GlobalComponents.Find(Name); Components && Components->IsValid())
+						if (TSharedPtr<FSubscriberImpl> Impl = WeakImpl.Pin())
 						{
-							FInstanceAudioComponents& ComponentsRef = (*Components->Get());
-							if (ComponentsRef.ActiveComponent.Get() == FinishingComponent)
+							if (FInstanceAudioComponentsPtr* ComponentsPtr = Impl->ComponentsStore.Find(FinishKey, Name))
 							{
-								ComponentsRef.ActiveComponent = nullptr;
-							}
-
-							ComponentsRef.ReleasedComponents.RemoveAllSwap([&FinishingComponent](const TObjectPtr<UAudioComponent>& Released)
-							{
-								return Released.Get() == FinishingComponent;
-							});
-
-							if (ComponentsRef.IsEmpty())
-							{
-								GlobalComponents.Remove(Name);
+								(*ComponentsPtr)->OnComponentFinished(FinishingComponent);
+								if ((*ComponentsPtr)->IsEmpty())
+								{
+									Impl->ComponentsStore.Remove(FinishKey, Name);
+								}
 							}
 						}
-					};
-					CacheComponent->OnAudioFinishedNative.AddSPLambda(this, RemoveComponentLambda);
-					return CacheComponent;
+					});
+					InstanceComponents->SetActiveComponent(ActiveComponent, MoveTemp(OnFinishedHandle));
 				}
 
-				return nullptr;
-			}
-
-			UAudioComponent* AddComponent(UObject& Parent, Audio::FDeviceId DeviceId, const Core::FCollectionHandle& InHandle, FName Name, bool bReleaseExisting)
-			{
-				return AddCachedComponent(Parent, DeviceId, InHandle, CollectionComponents, Name, bReleaseExisting);
-			}
-
-			UAudioComponent* AddComponent(UObject& Parent, Audio::FDeviceId DeviceId, const Core::FEventHandle& InHandle, FName Name, bool bReleaseExisting)
-			{
-				return AddCachedComponent(Parent, DeviceId, InHandle, EventComponents, Name, bReleaseExisting);
-			}
-
-			UAudioComponent* AddComponent(UObject& Parent, Audio::FDeviceId DeviceId, const Core::FActionHandle& InHandle, FName Name, bool bReleaseExisting)
-			{
-				return AddCachedComponent(Parent, DeviceId, InHandle, ActionComponents, Name, bReleaseExisting);
-			}
-
-			bool ContainsComponent(FName Name) const
-			{
-				return GlobalComponents.Contains(Name);
-			}
-
-			bool ContainsComponent(const Core::FCollectionHandle& InHandle, FName Name) const
-			{
-				return ContainsCacheComponent(CollectionComponents, InHandle, Name);
-			}
-
-			bool ContainsComponent(const Core::FEventHandle& InHandle, FName Name) const
-			{
-				return ContainsCacheComponent(EventComponents, InHandle, Name);
-			}
-
-			bool ContainsComponent(const Core::FActionHandle& InHandle, FName Name) const
-			{
-				return ContainsCacheComponent(ActionComponents, InHandle, Name);
+				return ActiveComponent;
 			}
 
 			UAudioComponent* FindComponent(FName Name) const
 			{
-				const FInstanceAudioComponentsPtr* InstanceComponents = GlobalComponents.Find(Name);
+				const FInstanceAudioComponentsPtr* InstanceComponents = ComponentsStore.Find(Name);
 				if (InstanceComponents && InstanceComponents->IsValid())
 				{
-					return (*InstanceComponents)->ActiveComponent.Get();
+					return (*InstanceComponents)->GetActiveComponent();
 				}
 
 				return nullptr;
 			}
 
-			UAudioComponent* FindComponent(const Core::FCollectionHandle& InHandle, FName Name) const
-			{
-				return FindCacheComponent(CollectionComponents, InHandle, Name);
-			}
-
-			UAudioComponent* FindComponent(const Core::FEventHandle& InHandle, FName Name) const
-			{
-				return FindCacheComponent(EventComponents, InHandle, Name);
-			}
-
-			UAudioComponent* FindComponent(const Core::FActionHandle& InHandle, FName Name) const
+			UAudioComponent* FindComponent(const Core::FExecutorScopeKey& InKey, FName Name) const
 			{
-				return FindCacheComponent(ActionComponents, InHandle, Name);
-			}
-
-			void OnCollectionUnregistered(const Core::FCollectionHandle& InHandle)
-			{
-				if (FAudioComponentsCache* Cache = CollectionComponents.Find(InHandle))
+				if (const FInstanceAudioComponentsPtr* ComponentsPtr = ComponentsStore.Find(InKey, Name))
 				{
-					for (TPair<FName, FInstanceAudioComponentsPtr>& Pair : *Cache)
-					{
-						FInstanceAudioComponentsPtr& InstanceComponents = Pair.Value;
-						check(InstanceComponents.IsValid());
-						if (UAudioComponent* Component = InstanceComponents->ActiveComponent.Get())
-						{
-							Component->Stop();
-						}
-
-						for (const TObjectPtr<UAudioComponent>& ReleasedComponent : InstanceComponents->ReleasedComponents)
-						{
-							if (UAudioComponent* Component = ReleasedComponent.Get())
-							{
-								Component->Stop();
-							}
-						}
-					}
-
-					CollectionComponents.Remove(InHandle);
-
-	#if SUBSONIC_ASSERT_ON_STALE_COLLECTION_EXECUTORS
-					for (TPair<uint32, Core::FCollectionHandle>& Pair : CollectionExecutors)
-					{
-						checkf(Pair.Value.CollectionId != InHandle.CollectionId, TEXT("Executor for Collection '%s' still exists beyond lifetime of collection."), *InHandle.ToString());
-					}
-	#endif // SUBSONIC_ASSERT_ON_STALE_COLLECTION_EXECUTORS
+					return (*ComponentsPtr)->GetActiveComponent();
 				}
+
+				return nullptr;
 			}
 
 			void OnExecutorRegistered(const Core::FSubsonicExecutor& InExecutor)
 			{
-				CollectionExecutors.Add(InExecutor.GetId(), InExecutor.GetCollectionHandle());
 			}
 
 			void OnExecutorUnregistered(const Core::FSubsonicExecutor& InExecutor)
 			{
-				if (const Core::FCollectionHandle* CollectionHandle = CollectionExecutors.Find(InExecutor.GetId()))
+				const Core::FExecutorScopeKey Key(InExecutor);
+				ComponentsStore.ForEach(Key, [this](FName, FInstanceAudioComponentsPtr& ComponentsPtr)
 				{
-					if (const FAudioComponentsCache* ScopeComponentMap = CollectionComponents.Find(*CollectionHandle))
-					{
-						for (const TPair<FName, FInstanceAudioComponentsPtr>& Pair : *ScopeComponentMap)
-						{
-							if (FInstanceAudioComponents* Components = Pair.Value.Get())
-							{
-								Components->StopAll();
-							}
-						}
-					}
-					CollectionExecutors.Remove(InExecutor.GetId());
-					CollectionComponents.Remove(*CollectionHandle);
-				}
+					check(ComponentsPtr.IsValid());
+					ComponentsPtr->StopAll();
+				});
+				ComponentsStore.Remove(Key);
 			}
 
 			void RemoveComponent(FName Name)
 			{
-				FInstanceAudioComponentsPtr* InstanceComponents = GlobalComponents.Find(Name);
-				if (InstanceComponents && InstanceComponents->IsValid())
+				if (FInstanceAudioComponentsPtr* ComponentsPtr = ComponentsStore.Find(Name))
 				{
-					(*InstanceComponents)->StopAll();
-					GlobalComponents.Remove(Name);
+					check(ComponentsPtr->IsValid());
+					(*ComponentsPtr)->StopAll();
+					ComponentsStore.Remove(Name);
 				}
 			}
 
-			void RemoveComponent(const Core::FCollectionHandle& InHandle, FName Name)
-			{
-				return RemoveCacheComponent(CollectionComponents, InHandle, Name);
-			}
-
-			void RemoveComponent(const Core::FEventHandle& InHandle, FName Name)
-			{
-				return RemoveCacheComponent(EventComponents, InHandle, Name);
-			}
-
-			void RemoveComponent(const Core::FActionHandle& InHandle, FName Name)
+			void RemoveComponent(const Core::FExecutorScopeKey& InKey, FName Name)
 			{
-				return RemoveCacheComponent(ActionComponents, InHandle, Name);
+				if (FInstanceAudioComponentsPtr* ComponentsPtr = ComponentsStore.Find(InKey, Name))
+				{
+					check(ComponentsPtr->IsValid());
+					(*ComponentsPtr)->StopAll();
+					ComponentsStore.Remove(InKey, Name);
+				}
 			}
 		};
 	} // namespace AudioComponentEventSubscriberPrivate
@@ -479,39 +349,9 @@ namespace UE::Subsonic
 		return Impl->AddComponent(*this, GetAudioDeviceId(), Name, bReleaseExisting);
 	}
 
-	UAudioComponent* USubsonicAudioComponentSubscriber::AddComponent(const Core::FCollectionHandle& InHandle, FName Name, bool bReleaseExisting)
-	{
-		return Impl->AddComponent(*this, GetAudioDeviceId(), InHandle, Name, bReleaseExisting);
-	}
-
-	UAudioComponent* USubsonicAudioComponentSubscriber::AddComponent(const Core::FEventHandle& InHandle, FName Name, bool bReleaseExisting)
+	UAudioComponent* USubsonicAudioComponentSubscriber::AddComponent(const Core::FExecutorScopeKey& InKey, FName Name, bool bReleaseExisting)
 	{
-		return Impl->AddComponent(*this, GetAudioDeviceId(), InHandle, Name, bReleaseExisting);
-	}
-
-	UAudioComponent* USubsonicAudioComponentSubscriber::AddComponent(const Core::FActionHandle& InHandle, FName Name, bool bReleaseExisting)
-	{
-		return Impl->AddComponent(*this, GetAudioDeviceId(), InHandle, Name, bReleaseExisting);
-	}
-
-	bool USubsonicAudioComponentSubscriber::ContainsComponent(FName Name) const
-	{
-		return Impl->ContainsComponent(Name);
-	}
-
-	bool USubsonicAudioComponentSubscriber::ContainsComponent(const Core::FCollectionHandle& InHandle, FName Name) const
-	{
-		return Impl->ContainsComponent(InHandle, Name);
-	}
-
-	bool USubsonicAudioComponentSubscriber::ContainsComponent(const Core::FEventHandle& InHandle, FName Name) const
-	{
-		return Impl->ContainsComponent(InHandle, Name);
-	}
-
-	bool USubsonicAudioComponentSubscriber::ContainsComponent(const Core::FActionHandle& InHandle, FName Name) const
-	{
-		return Impl->ContainsComponent(InHandle, Name);
+		return Impl->AddComponent(*this, GetAudioDeviceId(), InKey, Name, bReleaseExisting);
 	}
 
 	UAudioComponent* USubsonicAudioComponentSubscriber::FindComponent(FName Name)
@@ -519,19 +359,9 @@ namespace UE::Subsonic
 		return Impl->FindComponent(Name);
 	}
 
-	UAudioComponent* USubsonicAudioComponentSubscriber::FindComponent(const Core::FCollectionHandle& InHandle, FName Name)
-	{
-		return Impl->FindComponent(InHandle, Name);
-	}
-
-	UAudioComponent* USubsonicAudioComponentSubscriber::FindComponent(const Core::FEventHandle& InHandle, FName Name)
-	{
-		return Impl->FindComponent(InHandle, Name);
-	}
-
-	UAudioComponent* USubsonicAudioComponentSubscriber::FindComponent(const Core::FActionHandle& InHandle, FName Name)
+	UAudioComponent* USubsonicAudioComponentSubscriber::FindComponent(const Core::FExecutorScopeKey& InKey, FName Name)
 	{
-		return Impl->FindComponent(InHandle, Name);
+		return Impl->FindComponent(InKey, Name);
 	}
 
 	void USubsonicAudioComponentSubscriber::RemoveComponent(FName Name)
@@ -539,19 +369,9 @@ namespace UE::Subsonic
 		return Impl->RemoveComponent(Name);
 	}
 
-	void USubsonicAudioComponentSubscriber::RemoveComponent(const Core::FCollectionHandle& InHandle, FName Name)
-	{
-		return Impl->RemoveComponent(InHandle, Name);
-	}
-
-	void USubsonicAudioComponentSubscriber::RemoveComponent(const Core::FEventHandle& InHandle, FName Name)
-	{
-		return Impl->RemoveComponent(InHandle, Name);
-	}
-
-	void USubsonicAudioComponentSubscriber::OnCollectionUnregistered(const Core::FCollectionHandle& InHandle)
+	void USubsonicAudioComponentSubscriber::RemoveComponent(const Core::FExecutorScopeKey& InKey, FName Name)
 	{
-		Impl->OnCollectionUnregistered(InHandle);
+		return Impl->RemoveComponent(InKey, Name);
 	}
 
 	void USubsonicAudioComponentSubscriber::OnExecutorRegistered(const Core::FSubsonicExecutor& InExecutor)
@@ -564,4 +384,3 @@ namespace UE::Subsonic
 		Impl->OnExecutorUnregistered(InExecutor);
 	}
 } // namespace UE::Subsonic
-#undef SUBSONIC_ASSERT_ON_STALE_COLLECTION_EXECUTORS
diff --git a/Engine/Plugins/Experimental/Subsonic/Source/SubsonicEngine/Private/StandardEventSubscribers/AudioComponentEventSubscriber.h b/Engine/Plugins/Experimental/Subsonic/Source/SubsonicEngine/Private/StandardEventSubscribers/AudioComponentEventSubscriber.h
--- a/Engine/Plugins/Experimental/Subsonic/Source/SubsonicEngine/Private/StandardEventSubscribers/AudioComponentEventSubscriber.h
+++ b/Engine/Plugins/Experimental/Subsonic/Source/SubsonicEngine/Private/StandardEventSubscribers/AudioComponentEventSubscriber.h
@@ -7,7 +7,7 @@
 #include "Containers/Map.h"
 #include "StandardEventSubscribers/SubsonicEventSubscriberBase.h"
 #include "SubsonicEventSubscriberInterface.h"
-#include "SubsonicHandles.h"
+#include "SubsonicExecutor.h"
 #include "Subsystems/AudioEngineSubsystem.h"
 #include "Templates/SharedPointer.h"
 
@@ -19,8 +19,6 @@
 // Forward Declarations
 namespace UE::Subsonic
 {
-	struct FSubsonicExecutor;
-
 	namespace AudioComponentEventSubscriberPrivate
 	{
 		class FSubscriberImpl;
@@ -40,28 +38,17 @@ namespace UE::Subsonic
 		UE_API virtual void Deinitialize() override;
 		//~ End USubsonicEventSubscriberBase interface
 
-		UE_INTERNAL virtual void OnCollectionUnregistered(const Core::FCollectionHandle& InEvent) override;
 		UE_INTERNAL virtual void OnExecutorRegistered(const Core::FSubsonicExecutor& InExecutor) override;
 		UE_INTERNAL virtual void OnExecutorUnregistered(const Core::FSubsonicExecutor& InExecutor) override;
 
 		UE_INTERNAL UAudioComponent* AddComponent(FName Name, bool bReleaseExisting = true);
-		UE_INTERNAL UAudioComponent* AddComponent(const Core::FCollectionHandle& InHandle, FName Name, bool bReleaseExisting = true);
-		UE_INTERNAL UAudioComponent* AddComponent(const Core::FEventHandle& InHandle, FName Name, bool bReleaseExisting = true);
-		UE_INTERNAL UAudioComponent* AddComponent(const Core::FActionHandle& InHandle, FName Name, bool bReleaseExisting = true);
-
-		UE_INTERNAL bool ContainsComponent(FName Name) const;
-		UE_INTERNAL bool ContainsComponent(const Core::FCollectionHandle& InHandle, FName Name) const;
-		UE_INTERNAL bool ContainsComponent(const Core::FEventHandle& InHandle, FName Name) const;
-		UE_INTERNAL bool ContainsComponent(const Core::FActionHandle& InHandle, FName Name) const;
+		UE_INTERNAL UAudioComponent* AddComponent(const Core::FExecutorScopeKey& InKey, FName Name, bool bReleaseExisting = true);
 
 		UE_INTERNAL UAudioComponent* FindComponent(FName Name);
-		UE_INTERNAL UAudioComponent* FindComponent(const Core::FCollectionHandle& InHandle, FName Name);
-		UE_INTERNAL UAudioComponent* FindComponent(const Core::FEventHandle& InHandle, FName Name);
-		UE_INTERNAL UAudioComponent* FindComponent(const Core::FActionHandle& InHandle, FName Name);
+		UE_INTERNAL UAudioComponent* FindComponent(const Core::FExecutorScopeKey& InKey, FName Name);
 
 		UE_INTERNAL void RemoveComponent(FName Name);
-		UE_INTERNAL void RemoveComponent(const Core::FCollectionHandle& InHandle, FName Name);
-		UE_INTERNAL void RemoveComponent(const Core::FEventHandle& InHandle, FName Name);
+		UE_INTERNAL void RemoveComponent(const Core::FExecutorScopeKey& InKey, FName Name);
 
 	private:
 		TSharedPtr<AudioComponentEventSubscriberPrivate::FSubscriberImpl> Impl;
diff --git a/Engine/Plugins/Experimental/Subsonic/Source/SubsonicEngine/Private/StandardEventSubscribers/SubsonicGeneratorSourceSubscriber.cpp b/Engine/Plugins/Experimental/Subsonic/Source/SubsonicEngine/Private/StandardEventSubscribers/SubsonicGeneratorSourceSubscriber.cpp
--- a/Engine/Plugins/Experimental/Subsonic/Source/SubsonicEngine/Private/StandardEventSubscribers/SubsonicGeneratorSourceSubscriber.cpp
+++ b/Engine/Plugins/Experimental/Subsonic/Source/SubsonicEngine/Private/StandardEventSubscribers/SubsonicGeneratorSourceSubscriber.cpp
@@ -77,19 +77,14 @@ namespace UE::Subsonic
 			return MakeShared<FSubsonicGenerator>(WaveGen, SampleRate);
 		}
 
-		template <typename HandleType>
-		Audio::IAudioMixerGeneratorSource* PlaySound(
+		void PlaySoundOnInstance(
 			USubsonicGeneratorSourceSubscriber& Subscriber,
 			USoundWave& Sound,
-			const HandleType& InHandle,
-			TMap<HandleType, FGeneratorSourceCache>& OutSourceCache,
-			FName Name,
-			const FSubsonicParameterStore& Params,
+			FInstanceSources& InstanceSources,
 			FSubsonicRelay& Relay,
+			const FSubsonicParameterStore& Params,
 			bool bReleaseExisting)
 		{
-			Audio::IAudioMixerGeneratorSource* ReturnSource = nullptr;
-
 			if (FAudioDeviceManager* AudioDeviceManager = FAudioDeviceManager::Get())
 			{
 				if (Audio::FMixerDevice* AudioDevice = Subscriber.GetMutableMixerDevice())
@@ -97,16 +92,12 @@ namespace UE::Subsonic
 					TSharedPtr<FSubsonicGenerator> Generator = CreateGenerator(Sound, *AudioDevice);
 					if (!Generator.IsValid())
 					{
-						return nullptr;
+						return;
 					}
 
-					FGeneratorSourceCache& Cache = OutSourceCache.FindOrAdd(InHandle);
-					FInstanceSources& InstanceSources = Cache.FindOrAdd(Name);
-
 					TUniquePtr<Audio::IAudioMixerGeneratorSource> NewSource = AudioDevice->CreateGeneratorSource(Generator.ToSharedRef());
 					if (NewSource.IsValid())
 					{
-						ReturnSource = NewSource.Get();
 						if (InstanceSources.ActiveSource != nullptr)
 						{
 							FReleasedSource Released;
@@ -121,121 +112,72 @@ namespace UE::Subsonic
 						InstanceSources.ApplyParameters(Relay);
 
 						Generator->SetSource(NewSource.Get());
-						Relay.EnqueueCommand(FRelayCommand{ FRelayCommand::EType::Play, nullptr, Generator });
+						Relay.EnqueueCommand(FRelayCommand { FRelayCommand::EType::Play, nullptr, Generator });
 						InstanceSources.ActiveSource = MoveTemp(NewSource);
 					}
 				}
 			}
-
-			return ReturnSource;
 		}
 
-		template <typename HandleType>
-		void ClearStaleSounds(TMap<HandleType, FGeneratorSourceCache>& SourceCache)
+		void StopSoundOnInstance(FInstanceSources& Instances, FSubsonicRelay& Relay)
 		{
-			TArray<HandleType> HandlesToClear;
-
-			for (TPair<HandleType, FGeneratorSourceCache>& CacheItem : SourceCache)
+			if (Instances.ActiveSource != nullptr)
 			{
-				TArray<FName> NamesToClear;
-				for (TPair<FName, FInstanceSources>& SourcesItem : CacheItem.Value)
-				{
-					FInstanceSources& Instances = SourcesItem.Value;
+				FReleasedSource Released;
+				Released.Source = MoveTemp(Instances.ActiveSource);
+				Released.Generator = Instances.Generator;
+				Instances.ReleasedSources.Add(MoveTemp(Released));
 
-					Instances.ReleasedSources.RemoveAll([](const FReleasedSource& Released)
-						{
-							return !Released.Generator || Released.Generator->GetSource() == nullptr;
-						});
-					if (Instances.ActiveSource != nullptr && Instances.Generator && Instances.Generator->GetSource() == nullptr)
-					{
-						Instances.ActiveSource.Reset();
-						Instances.Generator.Reset();
-					}
-					if (Instances.ActiveSource == nullptr && Instances.ReleasedSources.IsEmpty())
-					{
-						NamesToClear.Add(SourcesItem.Key);
-					}
-				}
-
-				for (FName Name : NamesToClear)
-				{
-					CacheItem.Value.Remove(Name);
-				}
-				if (CacheItem.Value.IsEmpty())
+				Relay.EnqueueCommand(FRelayCommand
 				{
-					HandlesToClear.Add(CacheItem.Key);
-				}
-			}
+					FRelayCommand::EType::Stop, nullptr, Instances.Generator
+				});
 
-			for (HandleType Handle : HandlesToClear)
-			{
-				SourceCache.Remove(Handle);
+				Instances.Generator.Reset();
 			}
 		}
 
-		template <typename HandleType>
-		void StopSound(
-			FSubsonicRelay& Relay,
-			const HandleType& InHandle,
-			TMap<HandleType, FGeneratorSourceCache>& OutSourceCache,
-			FName Name)
+		void StopAllOnInstance(FInstanceSources& Instances, FSubsonicRelay& Relay)
 		{
-			FGeneratorSourceCache* Cache = OutSourceCache.Find(InHandle);
-			if (Cache != nullptr)
+			if (Instances.Generator)
 			{
-				if (FInstanceSources* Instances = Cache->Find(Name))
+				Relay.EnqueueCommand(FRelayCommand
 				{
-					if (Instances->ActiveSource != nullptr)
-					{
-						FReleasedSource Released;
-						Released.Source = MoveTemp(Instances->ActiveSource);
-						Released.Generator = Instances->Generator;
-						Instances->ReleasedSources.Add(MoveTemp(Released));
-
-						Relay.EnqueueCommand(FRelayCommand{
-							FRelayCommand::EType::Stop, nullptr, Instances->Generator
-						});
-
-						Instances->Generator.Reset();
-					}
-				}
+					FRelayCommand::EType::Stop, nullptr, Instances.Generator
+				});
 			}
-		}
 
-		template <typename HandleType>
-		void SetParameters(
-			FSubsonicRelay& Relay,
-			const HandleType& InHandle,
-			TMap<HandleType, FGeneratorSourceCache>& SourceCache,
-			FName SourceName,
-			const FSubsonicParameterStore& Params)
-		{
-			if (FGeneratorSourceCache* Cache = SourceCache.Find(InHandle))
+			for (FReleasedSource& Released : Instances.ReleasedSources)
 			{
-				if (FInstanceSources* Instances = Cache->Find(SourceName))
+				if (Released.Generator)
 				{
-					Instances->RuntimeParams.MergeFrom(Params);
-					Instances->ApplyParameters(Relay);
+					Relay.EnqueueCommand(FRelayCommand
+					{
+						FRelayCommand::EType::Stop, nullptr, Released.Generator
+					});
 				}
 			}
+
+			Instances.ActiveSource.Reset();
+			Instances.Generator.Reset();
+			Instances.ReleasedSources.Empty();
 		}
 
-		template <typename HandleType>
-		void SetAuthoredParameters(
-			FSubsonicRelay& Relay,
-			const HandleType& InHandle,
-			TMap<HandleType, FGeneratorSourceCache>& SourceCache,
-			FName SourceName,
-			const FSubsonicParameterStore& Params)
+		// Returns true if the instance is fully empty and can be removed from the store.
+		bool ClearStaleSoundsOnInstance(FInstanceSources& Instances)
 		{
-			if (FGeneratorSourceCache* Cache = SourceCache.Find(InHandle))
+			Instances.ReleasedSources.RemoveAll([](const FReleasedSource& Released)
 			{
-				if (FInstanceSources* Instances = Cache->Find(SourceName))
-				{
-					Instances->AuthoredParams = Params;
-					Instances->ApplyParameters(Relay);
-				}
+				return !Released.Generator || Released.Generator->GetSource() == nullptr;
+			});
+
+			if (Instances.ActiveSource != nullptr && Instances.Generator && Instances.Generator->GetSource() == nullptr)
+			{
+				Instances.ActiveSource.Reset();
+				Instances.Generator.Reset();
 			}
+
+			return Instances.ActiveSource == nullptr && Instances.ReleasedSources.IsEmpty();
 		}
 	} // namespace SubscriberPrivate
 
@@ -255,10 +197,11 @@ namespace UE::Subsonic
 	{
 		SetTickableTickType(ETickableTickType::Never);
 
-		GlobalSources.Reset();
-		CollectionSources.Reset();
-		EventSources.Reset();
-		ActionSources.Reset();
+		auto StopAll = [this](FName, FInstanceSources& Instances)
+		{
+			SubscriberPrivate::StopAllOnInstance(Instances, *Relay);
+		};
+		SourcesDataStore.Empty(StopAll);
 
 		Super::Deinitialize();
 	}
@@ -269,91 +212,84 @@ namespace UE::Subsonic
 
 	void USubsonicGeneratorSourceSubscriber::OnCollectionUnregistered(const FCollectionHandle& InHandle)
 	{
-		CollectionSources.Remove(InHandle);
 	}
 
-	void USubsonicGeneratorSourceSubscriber::PlaySound(FName Name, USoundWave& Sound, const FSubsonicParameterStore& Params, bool bReleaseExisting)
+	void USubsonicGeneratorSourceSubscriber::OnExecutorRegistered(const Core::FSubsonicExecutor& InExecutor)
 	{
-		const int32 FakeHandle = 0;
-		SubscriberPrivate::PlaySound(*this, Sound, FakeHandle, GlobalSources, Name, Params, *Relay, bReleaseExisting);
 	}
 
-	void USubsonicGeneratorSourceSubscriber::PlaySound(const FCollectionHandle& InHandle, FName Name, USoundWave& Sound, const FSubsonicParameterStore& Params, bool bReleaseExisting)
+	void USubsonicGeneratorSourceSubscriber::OnExecutorUnregistered(const Core::FSubsonicExecutor& InExecutor)
 	{
-		SubscriberPrivate::PlaySound(*this, Sound, InHandle, CollectionSources, Name, Params, *Relay, bReleaseExisting);
+		const FExecutorScopeKey Key(InExecutor);
+		SourcesDataStore.ForEach(Key, [this](FName, FInstanceSources& Instances)
+		{
+			SubscriberPrivate::StopAllOnInstance(Instances, *Relay);
+		});
+		SourcesDataStore.Remove(Key);
 	}
 
-	void USubsonicGeneratorSourceSubscriber::PlaySound(const FEventHandle& InHandle, FName Name, USoundWave& Sound, const FSubsonicParameterStore& Params, bool bReleaseExisting)
+	void USubsonicGeneratorSourceSubscriber::PlaySound(FName Name, USoundWave& Sound, const FSubsonicParameterStore& Params, bool bReleaseExisting)
 	{
-		SubscriberPrivate::PlaySound(*this, Sound, InHandle, EventSources, Name, Params, *Relay, bReleaseExisting);
+		FInstanceSources& Instances = SourcesDataStore.FindOrAdd(Name);
+		SubscriberPrivate::PlaySoundOnInstance(*this, Sound, Instances, *Relay, Params, bReleaseExisting);
 	}
 
-	void USubsonicGeneratorSourceSubscriber::PlaySound(const FActionHandle& InHandle, FName Name, USoundWave& Sound, const FSubsonicParameterStore& Params, bool bReleaseExisting)
+	void USubsonicGeneratorSourceSubscriber::PlaySound(const FExecutorScopeKey& InKey, FName Name, USoundWave& Sound, const FSubsonicParameterStore& Params, bool bReleaseExisting)
 	{
-		SubscriberPrivate::PlaySound(*this, Sound, InHandle, ActionSources, Name, Params, *Relay, bReleaseExisting);
+		FInstanceSources& Instances = SourcesDataStore.FindOrAdd(InKey, Name);
+		SubscriberPrivate::PlaySoundOnInstance(*this, Sound, Instances, *Relay, Params, bReleaseExisting);
 	}
 
 	void USubsonicGeneratorSourceSubscriber::StopSound(FName Name)
 	{
-		const int32 FakeHandle = 0;
-		SubscriberPrivate::StopSound(*Relay, FakeHandle, GlobalSources, Name);
-	}
-
-	void USubsonicGeneratorSourceSubscriber::StopSound(const FCollectionHandle& InHandle, FName Name)
-	{
-		SubscriberPrivate::StopSound(*Relay, InHandle, CollectionSources, Name);
-	}
-
-	void USubsonicGeneratorSourceSubscriber::StopSound(const FEventHandle& InHandle, FName Name)
-	{
-		SubscriberPrivate::StopSound(*Relay, InHandle, EventSources, Name);
+		if (FInstanceSources* Instances = SourcesDataStore.Find(Name))
+		{
+			SubscriberPrivate::StopSoundOnInstance(*Instances, *Relay);
+		}
 	}
 
-	void USubsonicGeneratorSourceSubscriber::StopSound(const FActionHandle& InHandle, FName Name)
+	void USubsonicGeneratorSourceSubscriber::StopSound(const FExecutorScopeKey& InKey, FName Name)
 	{
-		SubscriberPrivate::StopSound(*Relay, InHandle, ActionSources, Name);
+		if (FInstanceSources* Instances = SourcesDataStore.Find(InKey, Name))
+		{
+			SubscriberPrivate::StopSoundOnInstance(*Instances, *Relay);
+		}
 	}
 
 	void USubsonicGeneratorSourceSubscriber::SetParameters(FName SourceName, const FSubsonicParameterStore& Params)
 	{
-		const int32 FakeHandle = 0;
-		SubscriberPrivate::SetParameters(*Relay, FakeHandle, GlobalSources, SourceName, Params);
-	}
-
-	void USubsonicGeneratorSourceSubscriber::SetParameters(const FCollectionHandle& InHandle, FName SourceName, const FSubsonicParameterStore& Params)
-	{
-		SubscriberPrivate::SetParameters(*Relay, InHandle, CollectionSources, SourceName, Params);
-	}
-
-	void USubsonicGeneratorSourceSubscriber::SetParameters(const FEventHandle& InHandle, FName SourceName, const FSubsonicParameterStore& Params)
-	{
-		SubscriberPrivate::SetParameters(*Relay, InHandle, EventSources, SourceName, Params);
-	}
-
-	void USubsonicGeneratorSourceSubscriber::SetParameters(const FActionHandle& InHandle, FName SourceName, const FSubsonicParameterStore& Params)
-	{
-		SubscriberPrivate::SetParameters(*Relay, InHandle, ActionSources, SourceName, Params);
-	}
-
-	void USubsonicGeneratorSourceSubscriber::SetAuthoredParameters(FName SourceName, const FSubsonicParameterStore& Params)
-	{
-		const int32 FakeHandle = 0;
-		SubscriberPrivate::SetAuthoredParameters(*Relay, FakeHandle, GlobalSources, SourceName, Params);
+		if (FInstanceSources* Instances = SourcesDataStore.Find(SourceName))
+		{
+			Instances->RuntimeParams.MergeFrom(Params);
+			Instances->ApplyParameters(*Relay);
+		}
 	}
 
-	void USubsonicGeneratorSourceSubscriber::SetAuthoredParameters(const FCollectionHandle& InHandle, FName SourceName, const FSubsonicParameterStore& Params)
+	void USubsonicGeneratorSourceSubscriber::SetParameters(const FExecutorScopeKey& InKey, FName SourceName, const FSubsonicParameterStore& Params)
 	{
-		SubscriberPrivate::SetAuthoredParameters(*Relay, InHandle, CollectionSources, SourceName, Params);
+		if (FInstanceSources* Instances = SourcesDataStore.Find(InKey, SourceName))
+		{
+			Instances->RuntimeParams.MergeFrom(Params);
+			Instances->ApplyParameters(*Relay);
+		}
 	}
 
-	void USubsonicGeneratorSourceSubscriber::SetAuthoredParameters(const FEventHandle& InHandle, FName SourceName, const FSubsonicParameterStore& Params)
+	void USubsonicGeneratorSourceSubscriber::SetAuthoredParameters(FName SourceName, FSubsonicParameterStore Params)
 	{
-		SubscriberPrivate::SetAuthoredParameters(*Relay, InHandle, EventSources, SourceName, Params);
+		if (FInstanceSources* Instances = SourcesDataStore.Find(SourceName))
+		{
+			Instances->AuthoredParams = MoveTemp(Params);
+			Instances->ApplyParameters(*Relay);
+		}
 	}
 
-	void USubsonicGeneratorSourceSubscriber::SetAuthoredParameters(const FActionHandle& InHandle, FName SourceName, const FSubsonicParameterStore& Params)
+	void USubsonicGeneratorSourceSubscriber::SetAuthoredParameters(const FExecutorScopeKey& InKey, FName SourceName, FSubsonicParameterStore Params)
 	{
-		SubscriberPrivate::SetAuthoredParameters(*Relay, InHandle, ActionSources, SourceName, Params);
+		if (FInstanceSources* Instances = SourcesDataStore.Find(InKey, SourceName))
+		{
+			Instances->AuthoredParams = MoveTemp(Params);
+			Instances->ApplyParameters(*Relay);
+		}
 	}
 
 	void USubsonicGeneratorSourceSubscriber::Tick(float DeltaTime)
@@ -363,10 +299,34 @@ namespace UE::Subsonic
 			Relay->Tick(*MixerDevice);
 		}
 
-		SubscriberPrivate::ClearStaleSounds(GlobalSources);
-		SubscriberPrivate::ClearStaleSounds(CollectionSources);
-		SubscriberPrivate::ClearStaleSounds(EventSources);
-		SubscriberPrivate::ClearStaleSounds(ActionSources);
+		// Clear stale sounds at global scope.
+		TArray<FName> GlobalNamesToClear;
+		SourcesDataStore.ForEach([&GlobalNamesToClear](FName Name, FInstanceSources& Instances)
+		{
+			if (SubscriberPrivate::ClearStaleSoundsOnInstance(Instances))
+			{
+				GlobalNamesToClear.Add(Name);
+			}
+		});
+		for (FName Name : GlobalNamesToClear)
+		{
+			SourcesDataStore.Remove(Name);
+		}
+
+		// Clear stale sounds at executor scope.
+		TArray<TPair<FExecutorScopeKey, FName>> ToRemove;
+		SourcesDataStore.ForEach([&ToRemove](const FExecutorScopeKey& Key, const FName& Name, FInstanceSources& Instances)
+		{
+			if (SubscriberPrivate::ClearStaleSoundsOnInstance(Instances))
+			{
+				ToRemove.Emplace(Key, Name);
+			}
+		});
+
+		for (const TPair<FExecutorScopeKey, FName>& Entry : ToRemove)
+		{
+			SourcesDataStore.Remove(Entry.Key, Entry.Value);
+		}
 	}
 
 	TStatId USubsonicGeneratorSourceSubscriber::GetStatId() const
diff --git a/Engine/Plugins/Experimental/Subsonic/Source/SubsonicEngine/Private/SubsonicEventCollectionObjects.cpp b/Engine/Plugins/Experimental/Subsonic/Source/SubsonicEngine/Private/SubsonicEventCollectionObjects.cpp
--- a/Engine/Plugins/Experimental/Subsonic/Source/SubsonicEngine/Private/SubsonicEventCollectionObjects.cpp
+++ b/Engine/Plugins/Experimental/Subsonic/Source/SubsonicEngine/Private/SubsonicEventCollectionObjects.cpp
@@ -118,7 +118,7 @@ namespace UE::Subsonic
 
 	void USubsonicEventCollectionExecutor::BeginDestroy()
 	{
-		Executor.Unregister();
+		Executor.Reset();
 		Super::BeginDestroy();
 	}
 
@@ -171,7 +171,7 @@ namespace UE::Subsonic
 
 	void USubsonicEventCollectionExecutor::ExecuteEvent(FGameplayTag EventTag, ESubsonicExecutionResult& OutResult)
 	{
-		const bool bExecuted = Executor.ExecuteEvent(EventTag.GetTagName());
+		const bool bExecuted = Executor.IsValid() && Executor->ExecuteEvent(EventTag.GetTagName());
 		OutResult = bExecuted ? ESubsonicExecutionResult::Succeeded : ESubsonicExecutionResult::Failed;
 	}
 
@@ -188,17 +188,23 @@ namespace UE::Subsonic
 
 	const Core::FSubsonicExecutor& USubsonicEventCollectionExecutor::GetExecutor() const
 	{
-		return Executor;
+		check(IsValid());
+		return *Executor.Get();
 	}
 
 	bool USubsonicEventCollectionExecutor::IsValid() const
 	{
-		return Executor.IsValid() && GetCollection() != nullptr;
+		return Executor.IsValid() && Executor->IsValid() && GetCollection() != nullptr;
 	}
 
 	void USubsonicEventCollectionExecutor::Unregister()
 	{
-		Executor.Unregister();
+		if (Executor.IsValid())
+		{
+			Executor->Unregister();
+			Executor.Reset();
+		}
+
 		Collection = nullptr;
 	}
 } // namespace UE::Subsonic
diff --git a/Engine/Plugins/Experimental/Subsonic/Source/SubsonicEngine/Private/SubsonicSubsystem.cpp b/Engine/Plugins/Experimental/Subsonic/Source/SubsonicEngine/Private/SubsonicSubsystem.cpp
--- a/Engine/Plugins/Experimental/Subsonic/Source/SubsonicEngine/Private/SubsonicSubsystem.cpp
+++ b/Engine/Plugins/Experimental/Subsonic/Source/SubsonicEngine/Private/SubsonicSubsystem.cpp
@@ -54,34 +54,29 @@ namespace UE::Subsonic
 		private:
 			void IterateSubscribersMatchingDevice(Audio::FDeviceId DeviceId, TFunctionRef<void(TScriptInterface<Core::ISubsonicEventSubscriberInterface>&)> Func)
 			{
-				auto GetSubscriberDeviceId = [](TScriptInterface<Core::ISubsonicEventSubscriberInterface>& Subscriber) -> TOptional<Audio::FDeviceId>
+				for (TScriptInterface<Core::ISubsonicEventSubscriberInterface>& Subscriber : EventSubscribers)
 				{
-					UObject* SubscriberObject = Subscriber.GetObject();
-					if (!SubscriberObject)
+					TOptional<Audio::FDeviceId> SubscriberDeviceId;
+					if (UObject* SubscriberObject = Subscriber.GetObject())
 					{
-						return { };
-					}
-
-					if (const UAudioEngineSubsystem* Subsystem = Cast<UAudioEngineSubsystem>(Subscriber.GetObject()))
-					{
-						return Subsystem->GetAudioDeviceId();
-					}
-
-					Audio::FDeviceId DeviceId = INDEX_NONE;
-					if (TryGetDeviceIdFromWorldObject(Subscriber.GetObject(), DeviceId))
-					{
-						return DeviceId;
-					}
-
-					return { };
-				};
-
-				for (TScriptInterface<Core::ISubsonicEventSubscriberInterface>& EventSubscriber : EventSubscribers)
-				{
-					const TOptional<Audio::FDeviceId> SubscriberDeviceId = GetSubscriberDeviceId(EventSubscriber);
-					if (SubscriberDeviceId == DeviceId && EventSubscriber.GetObject() != nullptr)
-					{
-						Func(EventSubscriber);
+						if (const UAudioEngineSubsystem* Subsystem = Cast<UAudioEngineSubsystem>(SubscriberObject))
+						{
+							if (DeviceId == Subsystem->GetAudioDeviceId())
+							{
+								Func(Subscriber);
+							}
+						}
+						else
+						{
+							Audio::FDeviceId WorldDeviceId = INDEX_NONE;
+							if (TryGetDeviceIdFromWorldObject(Subscriber.GetObject(), WorldDeviceId))
+							{
+								if (WorldDeviceId == DeviceId)
+								{
+									Func(Subscriber);
+								}
+							}
+						}
 					}
 				}
 			}
@@ -175,6 +170,7 @@ namespace UE::Subsonic
 	void USubsonicSubsystem::Initialize(FSubsystemCollectionBase& Collection)
 	{
 		Super::Initialize(Collection);
+
 		Core::ISubsonicEventRegistry::Initialize(TUniquePtr<Core::ISubsonicEventRegistry>(new SubsystemPrivate::FSubsonicEventRegistry()));
 	}
 
diff --git a/Engine/Plugins/Experimental/Subsonic/Source/SubsonicEngine/Public/StandardActions/SubsonicAction_GeneratorSource.h b/Engine/Plugins/Experimental/Subsonic/Source/SubsonicEngine/Public/StandardActions/SubsonicAction_GeneratorSource.h
--- a/Engine/Plugins/Experimental/Subsonic/Source/SubsonicEngine/Public/StandardActions/SubsonicAction_GeneratorSource.h
+++ b/Engine/Plugins/Experimental/Subsonic/Source/SubsonicEngine/Public/StandardActions/SubsonicAction_GeneratorSource.h
@@ -11,7 +11,7 @@
 
 namespace UE::Subsonic
 {
-	USTRUCT(MinimalAPI, BlueprintType, DisplayName = "Play Sound (GeneratorSource)")
+	USTRUCT(MinimalAPI, DisplayName = "Play Sound (GeneratorSource)")
 		struct FSubsonicEventAction_GeneratorSourcePlay : public Core::FSubsonicEventActionBase
 	{
 		GENERATED_BODY()
@@ -20,27 +20,32 @@ namespace UE::Subsonic
 
 		virtual void Execute(const Core::FSubsonicExecutor& InExecutor, const Core::FActionHandle& InHandle) const override;
 
+#if WITH_EDITOR
+		virtual FText GetDisplayInfo() const override;
+#endif // WITH_EDITOR
+
 	public:
 		// Addressable name of Generator source instance.
 		UPROPERTY(EditAnywhere, Category = "Subsonic")
-		FName Name = "Sound01";
+		FName Name;
 
 		// Sound to play.
 		UPROPERTY(EditAnywhere, Category = "Subsonic")
 		TObjectPtr<USoundWave> Sound;
 
-		// Scope of source accessor.  If set to event, provided name is used
-		// to address the same source within the event, collection the same source
-		// within the collection, etc.
+		// Scope of source accessor.  If set to executor, named source is
+		// accessible only by pool of sibling event actions for a given
+		// executor instance. If set to global, accesses any named source
+		// in the global pool.
 		UPROPERTY(EditAnywhere, Category = "Subsonic")
-		ESubsonicExecutionScope Scope = ESubsonicExecutionScope::Collection;
+		ESubsonicExecutionScope Scope = ESubsonicExecutionScope::Executor;
 
 		// Authored parameters applied to the generator at play time (lowest priority layer).
 		UPROPERTY(EditAnywhere, Category = "Parameters")
 		FSubsonicParameterStore Parameters;
 	};
 
-	USTRUCT(MinimalAPI, BlueprintType, DisplayName = "Stop Sound (GeneratorSource)")
+	USTRUCT(MinimalAPI, DisplayName = "Stop Sound (GeneratorSource)")
 	struct FSubsonicEventAction_GeneratorSourceStop : public Core::FSubsonicEventActionBase
 	{
 		GENERATED_BODY()
@@ -49,15 +54,20 @@ namespace UE::Subsonic
 
 		virtual void Execute(const Core::FSubsonicExecutor& InExecutor, const Core::FActionHandle& InHandle) const override;
 
+#if WITH_EDITOR
+		virtual FText GetDisplayInfo() const override;
+#endif // WITH_EDITOR
+
 	public:
 		// Addressable name of source instance.
 		UPROPERTY(EditAnywhere, Category = "Subsonic")
-		FName Name = "Sound01";
+		FName Name;
 
-		// Scope of source accessor.  If set to event, provided name is used
-		// to address the same source within the event, collection the same source
-		// within the collection, etc.
+		// Scope of source accessor.  If set to executor, named source is
+		// accessible only by pool of sibling event actions for a given
+		// executor instance. If set to global, accesses any named source
+		// in the global pool.
 		UPROPERTY(EditAnywhere, Category = "Subsonic")
-		ESubsonicExecutionScope Scope = ESubsonicExecutionScope::Collection;
+		ESubsonicExecutionScope Scope = ESubsonicExecutionScope::Executor;
 	};
 } // namespace UE::Subsonic
\ No newline at end of file
diff --git a/Engine/Plugins/Experimental/Subsonic/Source/SubsonicEngine/Public/StandardEventSubscribers/SubsonicGeneratorSourceSubscriber.h b/Engine/Plugins/Experimental/Subsonic/Source/SubsonicEngine/Public/StandardEventSubscribers/SubsonicGeneratorSourceSubscriber.h
--- a/Engine/Plugins/Experimental/Subsonic/Source/SubsonicEngine/Public/StandardEventSubscribers/SubsonicGeneratorSourceSubscriber.h
+++ b/Engine/Plugins/Experimental/Subsonic/Source/SubsonicEngine/Public/StandardEventSubscribers/SubsonicGeneratorSourceSubscriber.h
@@ -4,8 +4,9 @@
 #include "IAudioMixerGeneratorSource.h"
 #include "StandardEventSubscribers/SubsonicEventSubscriberBase.h"
 #include "StandardEventSubscribers/SubsonicRelay.h"
-#include "SubsonicHandles.h"
+#include "SubsonicExecutor.h"
 #include "SubsonicParameterStore.h"
+#include "SubsonicSubscriberDataStore.h"
 #include "Tickable.h"
 
 #include "SubsonicGeneratorSourceSubscriber.generated.h"
@@ -45,8 +46,6 @@ namespace UE::Subsonic
 		void ApplyParameters(FSubsonicRelay& Relay);
 	};
 
-	using FGeneratorSourceCache = TMap<FName, FInstanceSources>;
-
 	/**
 	 * Interface to GeneratorSource playback Subsonic implementation.
 	 */
@@ -56,8 +55,7 @@ namespace UE::Subsonic
 		GENERATED_BODY()
 
 	public:
-		using FActionHandle = Core::FActionHandle;
-		using FEventHandle = Core::FEventHandle;
+		using FExecutorScopeKey = Core::FExecutorScopeKey;
 		using FCollectionHandle = Core::FCollectionHandle;
 
 		USubsonicGeneratorSourceSubscriber();
@@ -70,25 +68,20 @@ namespace UE::Subsonic
 		virtual void OnCollectionRegistered(const FCollectionHandle& InEvent) override;
 		virtual void OnCollectionUnregistered(const FCollectionHandle& InEvent) override;
 
+		virtual void OnExecutorRegistered(const Core::FSubsonicExecutor& InExecutor) override;
+		virtual void OnExecutorUnregistered(const Core::FSubsonicExecutor& InExecutor) override;
+
 		void PlaySound(FName Name, USoundWave& Sound, const FSubsonicParameterStore& Params, bool bReleaseExisting = true);
-		void PlaySound(const FCollectionHandle& InHandle, FName Name, USoundWave& Sound, const FSubsonicParameterStore& Params, bool bReleaseExisting = true);
-		void PlaySound(const FEventHandle& InHandle, FName Name, USoundWave& Sound, const FSubsonicParameterStore& Params, bool bReleaseExisting = true);
-		void PlaySound(const FActionHandle& InHandle, FName Name, USoundWave& Sound, const FSubsonicParameterStore& Params, bool bReleaseExisting = true);
+		void PlaySound(const FExecutorScopeKey& InKey, FName Name, USoundWave& Sound, const FSubsonicParameterStore& Params, bool bReleaseExisting = true);
 
 		void StopSound(FName Name);
-		void StopSound(const FCollectionHandle& InHandle, FName Name);
-		void StopSound(const FEventHandle& InHandle, FName Name);
-		void StopSound(const FActionHandle& InHandle, FName Name);
+		void StopSound(const FExecutorScopeKey& InKey, FName Name);
 
 		void SetParameters(FName SourceName, const FSubsonicParameterStore& Params);
-		void SetParameters(const FCollectionHandle& InHandle, FName SourceName, const FSubsonicParameterStore& Params);
-		void SetParameters(const FEventHandle& InHandle, FName SourceName, const FSubsonicParameterStore& Params);
-		void SetParameters(const FActionHandle& InHandle, FName SourceName, const FSubsonicParameterStore& Params);
+		void SetParameters(const FExecutorScopeKey& InKey, FName SourceName, const FSubsonicParameterStore& Params);
 
-		SUBSONICENGINE_API void SetAuthoredParameters(FName SourceName, const FSubsonicParameterStore& Params);
-		SUBSONICENGINE_API void SetAuthoredParameters(const FCollectionHandle& InHandle, FName SourceName, const FSubsonicParameterStore& Params);
-		SUBSONICENGINE_API void SetAuthoredParameters(const FEventHandle& InHandle, FName SourceName, const FSubsonicParameterStore& Params);
-		SUBSONICENGINE_API void SetAuthoredParameters(const FActionHandle& InHandle, FName SourceName, const FSubsonicParameterStore& Params);
+		SUBSONICENGINE_API void SetAuthoredParameters(FName SourceName, FSubsonicParameterStore Params);
+		SUBSONICENGINE_API void SetAuthoredParameters(const FExecutorScopeKey& InKey, FName SourceName, FSubsonicParameterStore Params);
 
 		// FTickableGameObject interface
 		virtual void Tick(float DeltaTime) override;
@@ -98,9 +91,6 @@ namespace UE::Subsonic
 	private:
 		TSharedRef<FSubsonicRelay> Relay;
 
-		TMap<int32, FGeneratorSourceCache> GlobalSources;
-		TMap<FCollectionHandle, FGeneratorSourceCache> CollectionSources;
-		TMap<FEventHandle, FGeneratorSourceCache> EventSources;
-		TMap<FActionHandle, FGeneratorSourceCache> ActionSources;
+		Core::TSubscriberDataStore<FInstanceSources> SourcesDataStore;
 	};
 } // namespace UE::Subsonic
diff --git a/Engine/Plugins/Experimental/Subsonic/Source/SubsonicEngine/Public/SubsonicDeviceUtils.h b/Engine/Plugins/Experimental/Subsonic/Source/SubsonicEngine/Public/SubsonicDeviceUtils.h
--- a/Engine/Plugins/Experimental/Subsonic/Source/SubsonicEngine/Public/SubsonicDeviceUtils.h
+++ b/Engine/Plugins/Experimental/Subsonic/Source/SubsonicEngine/Public/SubsonicDeviceUtils.h
@@ -7,10 +7,10 @@
 
 namespace UE::Subsonic
 {
-	// Utility to find an audio device subsystem from a scope key.
-	// Extracted from TScopeKey to keep SubsonicCore free of AudioDevice dependencies.
-	template <typename SubsystemType, typename ScopeHandleType>
-	SubsystemType* FindDeviceSubsystem(const Core::TScopeKey<ScopeHandleType>& InKey)
+	// Utility to find an audio device subsystem from an executor scope key. Free function here
+	// in SubsonicEngine to keep SubsonicCore free of AudioDevice dependencies.
+	template <typename SubsystemType>
+	SubsystemType* FindDeviceSubsystem(const Core::FExecutorScopeKey& InKey)
 	{
 		const FAudioDevice* AudioDevice = nullptr;
 		if (const FAudioDeviceManager* DeviceManager = FAudioDeviceManager::Get())
diff --git a/Engine/Plugins/Experimental/Subsonic/Source/SubsonicEngine/Public/SubsonicEventCollectionObjects.h b/Engine/Plugins/Experimental/Subsonic/Source/SubsonicEngine/Public/SubsonicEventCollectionObjects.h
--- a/Engine/Plugins/Experimental/Subsonic/Source/SubsonicEngine/Public/SubsonicEventCollectionObjects.h
+++ b/Engine/Plugins/Experimental/Subsonic/Source/SubsonicEngine/Public/SubsonicEventCollectionObjects.h
@@ -92,7 +92,7 @@ namespace UE::Subsonic
 		UPROPERTY(Transient, BlueprintReadOnly, Category = "Subsonic", meta = (AllowPrivateAccess))
 		TObjectPtr<const USubsonicEventCollection> Collection;
 
-		Core::FSubsonicExecutor Executor = Core::FSubsonicExecutor::GetInvalid();
+		TSharedPtr<Core::FSubsonicExecutor> Executor;
 	};
 } // namespace UE::Subsonic
 
diff --git a/Engine/Plugins/Experimental/Subsonic/Source/SubsonicEngineTest/Private/SubsonicEventCollectionObjectsTest.cpp b/Engine/Plugins/Experimental/Subsonic/Source/SubsonicEngineTest/Private/SubsonicEventCollectionObjectsTest.cpp
--- a/Engine/Plugins/Experimental/Subsonic/Source/SubsonicEngineTest/Private/SubsonicEventCollectionObjectsTest.cpp
+++ b/Engine/Plugins/Experimental/Subsonic/Source/SubsonicEngineTest/Private/SubsonicEventCollectionObjectsTest.cpp
@@ -57,8 +57,12 @@ namespace UE::Subsonic::Test
 				if (Executor)
 				{
 					ASSERT_THAT(IsTrue(Executor->IsValid()));
-					const bool bExecuted = Executor->GetExecutor().ExecuteEvent(Core::TAG_SubsonicCore_Event_Play.GetTag().GetTagName());
-					ASSERT_THAT(IsTrue(bExecuted));
+
+					if (Executor->IsValid())
+					{
+						const bool bExecuted = Executor->GetExecutor().ExecuteEvent(Core::TAG_SubsonicCore_Event_Play.GetTag().GetTagName());
+						ASSERT_THAT(IsTrue(bExecuted));
+					}		
 				}
 			}
 		}
